p4 integrate //UE5/Dev-Distro-5.5/Engine/Source/Programs/BuildPatchTool/... //UE5/Main/Engine/Source/Programs/BuildPatchTool/...

diff --git a/Engine/Source/Runtime/Online/BuildPatchServices/Private/BuildPatchInstaller.cpp b/Engine/Source/Runtime/Online/BuildPatchServices/Private/BuildPatchInstaller.cpp
--- a/Engine/Source/Runtime/Online/BuildPatchServices/Private/BuildPatchInstaller.cpp
+++ b/Engine/Source/Runtime/Online/BuildPatchServices/Private/BuildPatchInstaller.cpp
@@ -6,52 +6,64 @@
 =============================================================================*/
 
 #include "BuildPatchInstaller.h"
-#include "IBuildManifestSet.h"
-#include "HAL/PlatformFileManager.h"
+
+#include "Algo/AnyOf.h"
+#include "Algo/Transform.h"
+#include "Containers/Queue.h"
+#include "Containers/Ticker.h"
 #include "HAL/FileManager.h"
+#include "HAL/PlatformFileManager.h"
 #include "HAL/RunnableThread.h"
-#include "Misc/Paths.h"
-#include "Misc/ScopeLock.h"
+#include "HttpModule.h"
+#include "Math/UnitConversion.h"
 #include "Misc/ConfigCacheIni.h"
 #include "Misc/FeedbackContext.h"
-#include "Math/UnitConversion.h"
-#include "Containers/Ticker.h"
-#include "HttpModule.h"
+#include "Misc/Paths.h"
+#include "Misc/ScopeLock.h"
+
 #include "BuildPatchFileConstructor.h"
 #include "BuildPatchServicesModule.h"
-#include "BuildPatchUtil.h"
 #include "BuildPatchServicesPrivate.h"
 #include "BuildPatchSettings.h"
-#include "Common/HttpManager.h"
-#include "Common/FileSystem.h"
+#include "BuildPatchUtil.h"
 #include "Common/ChunkDataSizeProvider.h"
-#include "Installer/InstallerError.h"
+#include "Common/Crypto.h"
+#include "Common/FileSystem.h"
+#include "Common/HttpManager.h"
+#include "Generation/ChunkDatabaseWriter.h"
+#include "Generation/PackageChunkData.h"
+#include "IBuildManifestSet.h"
+#include "Installer/ChunkDbChunkSource.h"
+#include "Installer/ChunkEvictionPolicy.h"
+#include "Installer/ChunkReferenceTracker.h"
 #include "Installer/CloudChunkSource.h"
+#include "Installer/CustomUninstallAction.h"
 #include "Installer/DownloadConnectionCount.h"
-#include "Installer/MemoryChunkStore.h"
-#include "Installer/ChunkReferenceTracker.h"
-#include "Installer/ChunkEvictionPolicy.h"
 #include "Installer/DownloadService.h"
-#include "Installer/ChunkDbChunkSource.h"
-#include "Installer/InstallChunkSource.h"
-#include "Installer/Verifier.h"
 #include "Installer/FileAttribution.h"
+#include "Installer/InstallChunkSource.h"
 #include "Installer/InstallerAnalytics.h"
+#include "Installer/InstallerError.h"
 #include "Installer/InstallerSharedContext.h"
-#include "Installer/Prerequisites.h"
 #include "Installer/MachineConfig.h"
+#include "Installer/MemoryChunkStore.h"
 #include "Installer/MessagePump.h"
 #include "Installer/OptimisedDelta.h"
-#include "Installer/Statistics/MemoryChunkStoreStatistics.h"
-#include "Installer/Statistics/DownloadServiceStatistics.h"
+#include "Installer/Prerequisites.h"
 #include "Installer/Statistics/ChunkDbChunkSourceStatistics.h"
-#include "Installer/Statistics/InstallChunkSourceStatistics.h"
 #include "Installer/Statistics/CloudChunkSourceStatistics.h"
+#include "Installer/Statistics/DownloadServiceStatistics.h"
+#include "Installer/Statistics/Experiments.h"
 #include "Installer/Statistics/FileConstructorStatistics.h"
-#include "Installer/Statistics/VerifierStatistics.h"
 #include "Installer/Statistics/FileOperationTracker.h"
+#include "Installer/Statistics/InstallChunkSourceStatistics.h"
+#include "Installer/Statistics/MemoryChunkStoreStatistics.h"
+#include "Installer/Statistics/VerifierStatistics.h"
+#include "Installer/UriProvider.h"
+#include "Installer/Verifier.h"
 
-DEFINE_LOG_CATEGORY_STATIC(LogBPSInstallerConfig, Log, All);
+DECLARE_LOG_CATEGORY_EXTERN(LogBPSInstallerConfig, Log, All);
+DEFINE_LOG_CATEGORY(LogBPSInstallerConfig);
 
 namespace ConfigHelpers
 {
@@ -78,10 +90,25 @@ namespace ConfigHelpers
 		return FMath::Clamp<float>(AverageTime, 1.0f, 30.0f);
 	}
 
-	float DownloadSpeedAverageTime()
+	float LoadDownloadSpeedAverageTime_Throttled()
+	{
+		float AverageTime = 10.0;
+		GConfig->GetFloat(TEXT("Portal.BuildPatch"), TEXT("DownloadSpeedAverageTime_Throttled"), AverageTime, GEngineIni);
+		return FMath::Clamp<float>(AverageTime, 1.0f, 30.0f);
+	}
+
+	uint64 LoadPreloadMaxChunkDbSize()
+	{
+		int32 PreloadMaxChunkDbSizeMB = 3000;
+		GConfig->GetInt(TEXT("Portal.BuildPatch"), TEXT("PreloadMaxChunkDbSizeMB"), PreloadMaxChunkDbSizeMB, GEngineIni);
+		return uint64(FMath::Clamp<int32>(PreloadMaxChunkDbSizeMB, 1, TNumericLimits<int32>::Max())) * 1000 * 1000;
+	}
+
+	float DownloadSpeedAverageTime(bool bIsThrottlingDownload)
 	{
 		static float AverageTime = LoadDownloadSpeedAverageTime();
-		return AverageTime;
+		static float AverageTimeThrottled = LoadDownloadSpeedAverageTime_Throttled();
+		return bIsThrottlingDownload ? AverageTimeThrottled : AverageTime;
 	}
 
 	int32 NumFileMoveRetries()
@@ -96,25 +123,25 @@ namespace ConfigHelpers
 		return InstallerRetries;
 	}
 
-	FOptimisedDeltaConfiguration BuildOptimisedDeltaConfig(const FBuildInstallerConfiguration& Config, const FBuildPatchInstallerAction& InstallerAction)
+	FOptimisedDeltaConfiguration BuildOptimisedDeltaConfig(const FBuildInstallerConfiguration& Config, const FBuildPatchInstallerAction& InstallerAction, const int32_t RetriesNumber)
 	{
 		check(InstallerAction.IsUpdate());
 		// The optimised delta can deal with getting dupe manifests so lets just allow that for easy config.
 		FOptimisedDeltaConfiguration OptimisedDeltaConfiguration(InstallerAction.GetSharedInstallOrCurrentManifest());
 		OptimisedDeltaConfiguration.SourceManifest = InstallerAction.TryGetSharedCurrentManifest();
-		OptimisedDeltaConfiguration.CloudDirectories = Config.CloudDirectories;
-		if (!InstallerAction.GetCloudSubdirectory().IsEmpty())
-		{
-			for (FString& CloudDirectory : OptimisedDeltaConfiguration.CloudDirectories)
-			{
-				CloudDirectory /= InstallerAction.GetCloudSubdirectory();
-			}
-		}
 		OptimisedDeltaConfiguration.DeltaPolicy = Config.DeltaPolicy;
 		OptimisedDeltaConfiguration.InstallMode = Config.InstallMode;
+		OptimisedDeltaConfiguration.RetriesNumber = RetriesNumber;
 		OptimisedDeltaConfiguration.DeltaFilenameTrailer = Config.DeltaFilenameTrailer;
 		return OptimisedDeltaConfiguration;
 	}
+
+	TArray<FString> EnumeratePreloadChunkDbs(const TCHAR* SearchPath)
+	{
+		TArray<FString> ChunkDbResults;
+		IFileManager::Get().FindFiles(ChunkDbResults, SearchPath, TEXT("chunkdb"));
+		return ChunkDbResults;
+	}
 }
 
 namespace InstallerHelpers
@@ -159,6 +186,8 @@ namespace InstallerHelpers
 		UE_LOG(LogBuildPatchServices, Log, TEXT("Build Stat: InitializeTime: %s"), *FPlatformTime::PrettyTime(BuildStats.InitializeTime));
 		UE_LOG(LogBuildPatchServices, Log, TEXT("Build Stat: TheoreticalDownloadTime: %s"), *FPlatformTime::PrettyTime(BuildStats.TheoreticalDownloadTime));
 		UE_LOG(LogBuildPatchServices, Log, TEXT("Build Stat: ConstructTime: %s"), *FPlatformTime::PrettyTime(BuildStats.ConstructTime));
+		UE_LOG(LogBuildPatchServices, Log, TEXT("Build Stat: TotalDecryptTime: %s"), *FPlatformTime::PrettyTime(BuildStats.TotalDecryptTime));
+		UE_LOG(LogBuildPatchServices, Log, TEXT("Build Stat: UninstallActionTime: %s"), *FPlatformTime::PrettyTime(BuildStats.UninstallActionTime));
 		UE_LOG(LogBuildPatchServices, Log, TEXT("Build Stat: MoveFromStageTime: %s"), *FPlatformTime::PrettyTime(BuildStats.MoveFromStageTime));
 		UE_LOG(LogBuildPatchServices, Log, TEXT("Build Stat: FileAttributesTime: %s"), *FPlatformTime::PrettyTime(BuildStats.FileAttributesTime));
 		UE_LOG(LogBuildPatchServices, Log, TEXT("Build Stat: VerifyTime: %s"), *FPlatformTime::PrettyTime(BuildStats.VerifyTime));
@@ -170,13 +199,13 @@ namespace InstallerHelpers
 		UE_LOG(LogBuildPatchServices, Log, TEXT("Build Stat: ProcessSuccess: %s"), BuildStats.ProcessSuccess ? TEXT("TRUE") : TEXT("FALSE"));
 		UE_LOG(LogBuildPatchServices, Log, TEXT("Build Stat: ErrorCode: %s"), *BuildStats.ErrorCode);
 		UE_LOG(LogBuildPatchServices, Log, TEXT("Build Stat: FailureReasonText: %s"), *BuildStats.FailureReasonText.BuildSourceString());
-		UE_LOG(LogBuildPatchServices, Log, TEXT("Build Stat: FailureType: %s"), *EnumToString(BuildStats.FailureType));
+		UE_LOG(LogBuildPatchServices, Log, TEXT("Build Stat: FailureType: %s"), LexToString(BuildStats.FailureType));
 		UE_LOG(LogBuildPatchServices, Log, TEXT("Build Stat: NumInstallRetries: %u"), BuildStats.NumInstallRetries);
 		UE_LOG(LogBuildPatchServices, Log, TEXT("Build Stat: MaxDiskSpaceNeededWhenDeletingChunkDbsIfRequested: %s"), *FText::AsNumber(BuildStats.MaxDiskSpaceNeededWhenDeletingChunkDbsIfRequested).ToString());
 		check(BuildStats.NumInstallRetries == BuildStats.RetryFailureTypes.Num() && BuildStats.NumInstallRetries == BuildStats.RetryErrorCodes.Num());
 		for (uint32 RetryIdx = 0; RetryIdx < BuildStats.NumInstallRetries; ++RetryIdx)
 		{
-			UE_LOG(LogBuildPatchServices, Log, TEXT("Build Stat: RetryFailureType %u: %s"), RetryIdx, *EnumToString(BuildStats.RetryFailureTypes[RetryIdx]));
+			UE_LOG(LogBuildPatchServices, Log, TEXT("Build Stat: RetryFailureType %u: %s"), RetryIdx, LexToString(BuildStats.RetryFailureTypes[RetryIdx]));
 			UE_LOG(LogBuildPatchServices, Log, TEXT("Build Stat: RetryErrorCodes %u: %s"), RetryIdx, *BuildStats.RetryErrorCodes[RetryIdx]);
 		}
 		UE_LOG(LogBuildPatchServices, Log, TEXT("Build Stat: FinalProgressValue: %f"), BuildStats.FinalProgress);
@@ -207,6 +236,10 @@ namespace InstallerHelpers
 		{
 			return TEXT("Repair");
 		}
+		else if (InstallerAction.IsVerifyOnly())
+		{
+			return TEXT("VerifyOnly");
+		}
 		else if (InstallerAction.IsUninstall())
 		{
 			return TEXT("Uninstall");
@@ -242,6 +275,11 @@ namespace InstallerHelpers
 			UE_LOG(LogBPSInstallerConfig, Log, TEXT("Build Config: ActionType: %s"), GetActionTypeLog(InstallerAction));
 			UE_LOG(LogBPSInstallerConfig, Log, TEXT("Build Config: CurrentManifest: %s"), *GetManifestLog(CurrentManifest));
 			UE_LOG(LogBPSInstallerConfig, Log, TEXT("Build Config: InstallManifest: %s"), *GetManifestLog(InstallManifest));
+			TSet<FGuid> NecessarySecrets = InstallManifest.IsValid() ? InstallManifest->GetNecessaryEncryptionSecretIds() : TSet<FGuid>();
+			for (const FGuid& NecessarySecret : NecessarySecrets)
+			{
+				UE_LOG(LogBPSInstallerConfig, Log, TEXT("Build Config: InstallManifest: NecessarySecret: %s"), *LexToString(NecessarySecret));
+			}
 			for (const FString& Tag : InstallerAction.GetInstallTags())
 			{
 				UE_LOG(LogBPSInstallerConfig, Log, TEXT("Build Config: InstallTags: %s"), *Tag);
@@ -270,17 +308,24 @@ namespace InstallerHelpers
 		{
 			UE_LOG(LogBPSInstallerConfig, Log, TEXT("Build Config: ChunkDatabaseFile: %s"), *DatabaseFile);
 		}
-		
+
 		for (const FString& CloudDirectory : InstallerConfiguration.CloudDirectories)
 		{
 			UE_LOG(LogBPSInstallerConfig, Log, TEXT("Build Config: CloudDirectories: %s"), *CloudDirectory);
 		}
 
+		for (const TPair<FGuid, TArray<uint8>>& EncryptionSecret : InstallerConfiguration.EncryptionSecrets)
+		{
+			UE_LOG(LogBPSInstallerConfig, Log, TEXT("Build Config: AvailableSecretId: %s"), *EncryptionSecret.Key.ToString());
+		}
+
 		UE_LOG(LogBPSInstallerConfig, Log, TEXT("Build Config: InstallMode: %s"), LexToString(InstallerConfiguration.InstallMode));
 		UE_LOG(LogBPSInstallerConfig, Log, TEXT("Build Config: VerifyMode: %s"), LexToString(InstallerConfiguration.VerifyMode));
 		UE_LOG(LogBPSInstallerConfig, Log, TEXT("Build Config: DeltaPolicy: %s"), LexToString(InstallerConfiguration.DeltaPolicy));
 		UE_LOG(LogBPSInstallerConfig, Log, TEXT("Build Config: bRunRequiredPrereqs: %s"), (InstallerConfiguration.bRunRequiredPrereqs) ? TEXT("true") : TEXT("false"));
 		UE_LOG(LogBPSInstallerConfig, Log, TEXT("Build Config: bAllowConcurrentExecution: %s"), (InstallerConfiguration.bAllowConcurrentExecution) ? TEXT("true") : TEXT("false"));
+		UE_LOG(LogBPSInstallerConfig, Log, TEXT("Build Config: DownloadRateLimitBps: %llu"), InstallerConfiguration.DownloadRateLimitBps);
+		UE_LOG(LogBPSInstallerConfig, Log, TEXT("Build Config: bRejectSymlinks: %s"), (InstallerConfiguration.bRejectSymlinks) ? TEXT("true") : TEXT("false"));
 	}
 
 	TSet<FGuid> GetMultipleReferencedChunks(IBuildManifestSet* ManifestSet)
@@ -316,10 +361,10 @@ namespace InstallerHelpers
 		using namespace BuildPatchServices;
 		switch (VerifyResult)
 		{
-			case EVerifyResult::FileMissing: return VerifyErrorCodes::FileMissing;
-			case EVerifyResult::OpenFileFailed: return VerifyErrorCodes::OpenFileFailed;
-			case EVerifyResult::HashCheckFailed: return VerifyErrorCodes::HashCheckFailed;
-			case EVerifyResult::FileSizeFailed: return VerifyErrorCodes::FileSizeFailed;
+		case EVerifyResult::FileMissing: return VerifyErrorCodes::FileMissing;
+		case EVerifyResult::OpenFileFailed: return VerifyErrorCodes::OpenFileFailed;
+		case EVerifyResult::HashCheckFailed: return VerifyErrorCodes::HashCheckFailed;
+		case EVerifyResult::FileSizeFailed: return VerifyErrorCodes::FileSizeFailed;
 		}
 
 		return VerifyErrorCodes::UnknownFail;
@@ -337,9 +382,9 @@ namespace InstallerHelpers
 		}
 	}
 
-	FOptimisedDeltaDependencies BuildOptimisedDeltaDependencies(const TUniquePtr<IDownloadService>& DownloadService)
+	FOptimisedDeltaDependencies BuildOptimisedDeltaDependencies(const TUniquePtr<IDownloadService>& DownloadService, TSharedRef<IUriProvider, ESPMode::ThreadSafe> UriProvider)
 	{
-		FOptimisedDeltaDependencies OptimisedDeltaDependencies;
+		FOptimisedDeltaDependencies OptimisedDeltaDependencies(MoveTemp(UriProvider));
 		OptimisedDeltaDependencies.DownloadService = DownloadService.Get();
 		return OptimisedDeltaDependencies;
 	}
@@ -408,26 +453,149 @@ namespace BuildPatchServices
 		}
 	};
 
-	static uint64 DetermineInstallMaxDiskSizeIfDeletingChunkDbs(const TArray<FString>& ChunkDbFiles, IBuildManifestSet* ManifestSet, IFileSystem* FileSystem, const FString& InstallDirectory, const TArray<FString>& CorruptFiles, bool bIsPrereqOnly);
+	static uint64 DetermineInstallMaxDiskSizeIfDeletingChunkDbs(const TArray<FString>& ChunkDbFiles, IBuildManifestSet* ManifestSet, IFileSystem* FileSystem, const FString& InstallDirectory, const TArray<FString>& CorruptFiles, EInstallMode InstallMode);
+
+	namespace ConstructionHelpers
+	{
+		FBuildInstallerConfiguration&& ProcessOverrides(FBuildInstallerConfiguration&& Configuration)
+		{
+			// Clean up directory members.
+			FPaths::NormalizeDirectoryName(Configuration.InstallDirectory);
+			FPaths::CollapseRelativeDirectories(Configuration.InstallDirectory);
+			FPaths::NormalizeDirectoryName(Configuration.StagingDirectory);
+			FPaths::CollapseRelativeDirectories(Configuration.StagingDirectory);
+			FPaths::NormalizeDirectoryName(Configuration.BackupDirectory);
+			FPaths::CollapseRelativeDirectories(Configuration.BackupDirectory);
+
+			// Force stage with RAW filenames, if legacy staging locations exist.
+			IFileManager& FileManager = IFileManager::Get();
+			const FString LegacyDataStagingDir(Configuration.StagingDirectory / TEXT("PatchData"));
+			const FString LegacyInstallStagingDir(Configuration.StagingDirectory / TEXT("Install"));
+			const FString LegacyMetaStagingDir(Configuration.StagingDirectory / TEXT("Meta"));
+			Configuration.bStageWithRawFilenames = Configuration.bStageWithRawFilenames
+				|| FileManager.DirectoryExists(*LegacyDataStagingDir)
+				|| FileManager.DirectoryExists(*LegacyInstallStagingDir)
+				|| FileManager.DirectoryExists(*LegacyMetaStagingDir);
+
+			// Create a default shared context.
+			if (!Configuration.SharedContext)
+			{
+				Configuration.SharedContext = FBuildInstallerSharedContextFactory::Create(TEXT("BuildPatchInstaller"));
+				const bool bUseChunkDBs = !Configuration.ChunkDatabaseFiles.IsEmpty();
+				const bool bHasCurrentManifest = Algo::AnyOf(Configuration.InstallerActions, [](const FBuildPatchInstallerAction& Elem) { return Elem.TryGetCurrentManifest() != nullptr; });
+				const uint32 NumExpectedThreads = Configuration.SharedContext->NumThreadsPerInstaller(bUseChunkDBs, bHasCurrentManifest);
+				Configuration.SharedContext->PreallocateThreads(NumExpectedThreads);
+			}
+
+			return MoveTemp(Configuration);
+		}
+
+		FString GetDataStagingDir(const FBuildInstallerConfiguration& Configuration)
+		{
+			FString DataStagingDir;
+			if (Configuration.bStageWithRawFilenames)
+			{
+				DataStagingDir = Configuration.StagingDirectory / TEXT("PatchData");
+			}
+			else
+			{
+				DataStagingDir = Configuration.StagingDirectory / TEXT("c"); // c for chunk
+			}
+			FPaths::NormalizeDirectoryName(DataStagingDir);
+			FPaths::CollapseRelativeDirectories(DataStagingDir);
+			return DataStagingDir;
+		}
+
+		FString GetInstallStagingDir(const FBuildInstallerConfiguration& Configuration)
+		{
+			FString InstallStagingDir;
+			if (Configuration.bStageWithRawFilenames)
+			{
+				InstallStagingDir = Configuration.StagingDirectory / TEXT("Install");
+			}
+			else
+			{
+				InstallStagingDir = Configuration.StagingDirectory / TEXT("f"); // f for files
+			}
+			FPaths::NormalizeDirectoryName(InstallStagingDir);
+			FPaths::CollapseRelativeDirectories(InstallStagingDir);
+			return InstallStagingDir;
+		}
+
+		FString GetMetaStagingDir(const FBuildInstallerConfiguration& Configuration)
+		{
+			FString MetaStagingDir;
+			if (Configuration.bStageWithRawFilenames)
+			{
+				MetaStagingDir = Configuration.StagingDirectory / TEXT("Meta");
+			}
+			else
+			{
+				MetaStagingDir = Configuration.StagingDirectory / TEXT("m"); // m for meta
+			}
+			FPaths::NormalizeDirectoryName(MetaStagingDir);
+			FPaths::CollapseRelativeDirectories(MetaStagingDir);
+			return MetaStagingDir;
+		}
+
+		TArray<FBuildPatchInstallerAction> UpCastInstallerActions(const TArray<FInstallerAction>& InInstallerActions)
+		{
+			TArray<FBuildPatchInstallerAction> InstallerActions;
+			for (const FInstallerAction& InstallerAction : InInstallerActions)
+			{
+				InstallerActions.Emplace(InstallerAction);
+			}
+			return InstallerActions;
+		}
+
+		TMultiMap<FString, FBuildPatchAppManifestRef>&& AppendInstallationInfo(TMultiMap<FString, FBuildPatchAppManifestRef>&& InstallationInfo, TArray<FBuildPatchInstallerAction>& InstallerActions, const FBuildInstallerConfiguration& Configuration)
+		{
+			for (const FBuildPatchInstallerAction& InstallerAction : InstallerActions)
+			{
+				// Make sure existing manifests are added to installation info.
+				FBuildPatchAppManifestPtr CurrentManifest = InstallerAction.TryGetSharedCurrentManifest();
+				if (CurrentManifest.IsValid())
+				{
+					InstallationInfo.Add(Configuration.InstallDirectory / InstallerAction.GetInstallSubdirectory(), CurrentManifest.ToSharedRef());
+				}
+			}
+			return MoveTemp(InstallationInfo);
+		}
+
+		TArray<FBuildPatchAppManifestPtr> ToManifestArray(const TArray<FBuildPatchInstallerAction>& InstallerActions)
+		{
+			TArray<FBuildPatchAppManifestPtr> ManifestAraray;
+			for (const FBuildPatchInstallerAction& InstallerAction : InstallerActions)
+			{
+				ManifestAraray.AddUnique(InstallerAction.TryGetSharedCurrentManifest());
+				ManifestAraray.AddUnique(InstallerAction.TryGetSharedInstallManifest());
+			}
+			return ManifestAraray;
+		}
+	}
 
 	/* FBuildPatchInstaller implementation
 	*****************************************************************************/
 	FBuildPatchInstaller::FBuildPatchInstaller(FBuildInstallerConfiguration InConfiguration, TMultiMap<FString, FBuildPatchAppManifestRef> InInstallationInfo, const FString& InLocalMachineConfigFile, TSharedPtr<IAnalyticsProvider> InAnalytics, FBuildPatchInstallerDelegate InStartDelegate, FBuildPatchInstallerDelegate InCompleteDelegate)
-		: SessionId(FGuid::NewGuid())
-		, Thread(nullptr)
+		: Configuration(ConstructionHelpers::ProcessOverrides(MoveTemp(InConfiguration)))
+		, SessionId(FGuid::NewGuid())
 		, StartDelegate(InStartDelegate)
 		, CompleteDelegate(InCompleteDelegate)
-		, Configuration(MoveTemp(InConfiguration))
-		, DataStagingDir(Configuration.StagingDirectory / TEXT("PatchData"))
-		, InstallStagingDir(Configuration.StagingDirectory / TEXT("Install"))
-		, MetaStagingDir(Configuration.StagingDirectory / TEXT("Meta"))
+		, InstallerActions(ConstructionHelpers::UpCastInstallerActions(Configuration.InstallerActions))
+		, DataStagingDir(ConstructionHelpers::GetDataStagingDir(Configuration))
+		, InstallStagingDir(ConstructionHelpers::GetInstallStagingDir(Configuration))
+		, MetaStagingDir(ConstructionHelpers::GetMetaStagingDir(Configuration))
 		, PreviousMoveMarker(Configuration.InstallDirectory / TEXT("$movedMarker"))
+		, HddSpaceReservationFile(DataStagingDir / TEXT("reserve"))
+		, InstallationInfo(ConstructionHelpers::AppendInstallationInfo(MoveTemp(InInstallationInfo), InstallerActions, Configuration))
 		, ThreadLock()
+		, Thread(nullptr)
 		, bSuccess(false)
 		, bIsRunning(false)
 		, bIsInited(false)
 		, bFirstInstallIteration(true)
 		, PreviousTotalDownloadRequired(0)
+		, bAreResultsEncrypted(false)
 		, BuildStats()
 		, BuildProgress()
 		, bIsPaused(false)
@@ -435,10 +603,10 @@ namespace BuildPatchServices
 		, FilesInstalled()
 		, TaggedFiles()
 		, FilesToConstruct()
-		, InstallationInfo(MoveTemp(InInstallationInfo))
 		, LocalMachineConfigFile(InLocalMachineConfigFile)
 		, HttpManager(FHttpManagerFactory::Create())
 		, FileSystem(FFileSystemFactory::Create())
+		, Crypto(FCryptoFactory::Create())
 		, Platform(FPlatformFactory::Create())
 		, InstallerError(FInstallerErrorFactory::Create())
 		, Analytics(MoveTemp(InAnalytics))
@@ -447,58 +615,20 @@ namespace BuildPatchServices
 		, DownloadSpeedRecorder(FSpeedRecorderFactory::Create())
 		, DiskReadSpeedRecorder(FSpeedRecorderFactory::Create())
 		, DiskWriteSpeedRecorder(FSpeedRecorderFactory::Create())
-		, ChunkDataSizeProvider(FChunkDataSizeProviderFactory::Create())
+		, ChunkDataSizeProvider(FChunkDataSizeProviderFactory::Create(ConstructionHelpers::ToManifestArray(InstallerActions)))
 		, DownloadServiceStatistics(FDownloadServiceStatisticsFactory::Create(DownloadSpeedRecorder.Get(), ChunkDataSizeProvider.Get(), InstallerAnalytics.Get()))
 		, ChunkDbChunkSourceStatistics(FChunkDbChunkSourceStatisticsFactory::Create(DiskReadSpeedRecorder.Get(), FileOperationTracker.Get()))
 		, InstallChunkSourceStatistics(FInstallChunkSourceStatisticsFactory::Create(DiskReadSpeedRecorder.Get(), InstallerAnalytics.Get(), FileOperationTracker.Get()))
 		, CloudChunkSourceStatistics(FCloudChunkSourceStatisticsFactory::Create(InstallerAnalytics.Get(), &BuildProgress, FileOperationTracker.Get()))
 		, FileConstructorStatistics(FFileConstructorStatisticsFactory::Create(DiskReadSpeedRecorder.Get(), DiskWriteSpeedRecorder.Get(), &BuildProgress, FileOperationTracker.Get()))
 		, VerifierStatistics(FVerifierStatisticsFactory::Create(DiskReadSpeedRecorder.Get(), &BuildProgress, FileOperationTracker.Get()))
-		, DownloadService(FDownloadServiceFactory::Create(HttpManager.Get(), FileSystem.Get(), DownloadServiceStatistics.Get(), InstallerAnalytics.Get()))
+		, DownloadService(FDownloadServiceFactory::Create(HttpManager, FileSystem, DownloadServiceStatistics, InstallerAnalytics))
 		, MessagePump(FMessagePumpFactory::Create())
 		, Controllables()
+		, DownloadRateLimitBps(Configuration.DownloadRateLimitBps)
 	{
 		InstallerError->RegisterForErrors([this]() { CancelInstall(); });
 		Controllables.Add(&BuildProgress);
-
-		// Convert the manifest collection to make other code cleaner.
-		FString InstallDirectory = Configuration.InstallDirectory;
-		FPaths::NormalizeDirectoryName(InstallDirectory);
-		FPaths::CollapseRelativeDirectories(InstallDirectory);
-
-		TArray<FBuildPatchAppManifestPtr> Manifests;
-		Manifests.Reserve(Configuration.InstallerActions.Num() * 2);
-
-		bool bHasCurrentManifest = false;
-		for (const FInstallerAction& InstallerAction : Configuration.InstallerActions)
-		{
-			FBuildPatchInstallerAction& BuildPatchInstallerAction = InstallerActions.Emplace_GetRef(InstallerAction);
-			// Make sure existing manifests are added to installation info.
-			if (BuildPatchInstallerAction.TryGetCurrentManifest())
-			{
-				bHasCurrentManifest = true;
-				InstallationInfo.Add(InstallDirectory / BuildPatchInstallerAction.GetInstallSubdirectory(), BuildPatchInstallerAction.GetSharedCurrentManifest());
-			}
-			// Cache chunk sizes too
-			if (FBuildPatchAppManifestPtr SharedCurrentManifest = BuildPatchInstallerAction.TryGetSharedCurrentManifest())
-			{
-				Manifests.Emplace(MoveTemp(SharedCurrentManifest));
-			}
-			if (FBuildPatchAppManifestPtr SharedInstallManifest = BuildPatchInstallerAction.TryGetSharedInstallManifest())
-			{
-				Manifests.Emplace(MoveTemp(SharedInstallManifest));
-			}
-		}
-
-		ChunkDataSizeProvider->AddManifestData(Manifests);
-
-		if (!Configuration.SharedContext)
-		{
-			Configuration.SharedContext = FBuildInstallerSharedContextFactory::Create(TEXT("BuildPatchInstaller"));
-			const bool bUseChunkDBs = !Configuration.ChunkDatabaseFiles.IsEmpty();
-			const uint32 NumExpectedThreads = Configuration.SharedContext->NumThreadsPerInstaller(bUseChunkDBs, bHasCurrentManifest);
-			Configuration.SharedContext->PreallocateThreads(NumExpectedThreads);
-		}
 	}
 
 	FBuildPatchInstaller::~FBuildPatchInstaller()
@@ -599,7 +729,7 @@ namespace BuildPatchServices
 		}
 
 		TArray<FDownload> Downloads = DownloadServiceStatistics->GetCurrentDownloads();
-		for (const FDownload& Download: Downloads)
+		for (const FDownload& Download : Downloads)
 		{
 			Output.Add(FString::Printf(TEXT("BPI %s downloaded %.2f MBytes / %.2f MBytes"),
 				*Download.Data,
@@ -610,19 +740,64 @@ namespace BuildPatchServices
 	}
 #endif
 
+	void FBuildPatchInstaller::SetDownloadBytesPerSecondLimit(uint64 InDownloadRateLimitBps)
+	{
+		UE::TUniqueLock _(DownloadRateLimitLock);
+
+
+		if (InDownloadRateLimitBps != DownloadRateLimitBps)
+		{
+			UE_LOG(LogBuildPatchServices, Log, TEXT("Build Config: InstallerId: %s: SetDownloadRateLimit: %llu"), *SessionId.ToString(), InDownloadRateLimitBps);
+			DownloadRateLimitBps = InDownloadRateLimitBps;
+
+			if (DownloadRateUpdatedFn)
+			{
+				DownloadRateUpdatedFn(DownloadRateLimitBps);
+			}
+		}
+
+
+	}
+
+	uint64 FBuildPatchInstaller::GetDownloadBytesPerSecondLimit() const
+	{
+		uint64 ReturnRate;
+		{
+			UE::TUniqueLock _(const_cast<UE::FMutex&>(DownloadRateLimitLock));
+			ReturnRate = DownloadRateLimitBps;
+		}
+
+		return ReturnRate;
+	}
+
 	bool FBuildPatchInstaller::StartInstallation()
 	{
 		if (Thread == nullptr)
 		{
 			// Start thread!
+			check(Configuration.SharedContext);
 			Thread = Configuration.SharedContext->CreateThread();
-			Thread->RunTask([this]{ Run(); });
+			Thread->RunTask([this] { Run(); });
 
 			StartDelegate.ExecuteIfBound(AsShared());
 		}
 		return Thread != nullptr;
 	}
 
+	template<typename T>
+	bool WaitForResultOrAbort(T& DataToWait, bool& bShouldAbort, float PeriodToCheckSeconds = 0.05 /* 50 ms by default */)
+	{
+		while (!DataToWait.IsReady())
+		{
+			FPlatformProcess::Sleep(PeriodToCheckSeconds);
+			if (bShouldAbort)
+			{
+				return false;
+			}
+		}
+		return true;
+	}
+
 	bool FBuildPatchInstaller::Initialize()
 	{
 		bool bInstallerInitSuccess = true;
@@ -632,7 +807,6 @@ namespace BuildPatchServices
 		for (const FBuildPatchInstallerAction& InstallerAction : InstallerActions)
 		{
 			TSet<FString> ValidTags;
-			ValidTags.Add(TEXT(""));
 			InstallerAction.GetInstallOrCurrentManifest().GetFileTagList(ValidTags);
 			ValidTags.Add(TEXT(""));
 			TSet<FString> InvalidTags = InstallerAction.GetInstallTags().Difference(ValidTags);
@@ -646,6 +820,40 @@ namespace BuildPatchServices
 			}
 		}
 
+		// Check the actions to ensure if VerifyOnly is used, then it is exclusively used.
+		bool bHasVerifyOnly = false;
+		bool bHasOtherActionType = false;
+		for (const FBuildPatchInstallerAction& InstallerAction : InstallerActions)
+		{
+			if (InstallerAction.IsVerifyOnly())
+			{
+				bHasVerifyOnly = true;
+			}
+			else
+			{
+				bHasOtherActionType = true;
+			}
+		}
+		if (bHasVerifyOnly && bHasOtherActionType)
+		{
+			UE_LOG(LogBuildPatchServices, Error, TEXT("Installer configuration: Verify only actions must not be mixed with others."));
+			InstallerError->SetError(EBuildPatchInstallError::InitializationError, InitializationErrorCodes::VerifyOnlyActionMixedWithOtherActions, 0,
+				NSLOCTEXT("BuildPatchInstallError", "VerifyOnlyActionMixed", "This installation could not continue due to a configuration issue. Please contact support."));
+			bInstallerInitSuccess = false;
+		}
+
+		if (!bHasVerifyOnly && !Configuration.CloudDirectories.Num())
+		{
+			UE_LOG(LogBuildPatchServices, Error, TEXT("Installer configuration: Cloud directories are missing."));
+			InstallerError->SetError(EBuildPatchInstallError::InitializationError, InitializationErrorCodes::EmptyCloudDirectories, 0,
+				NSLOCTEXT("BuildPatchInstallError", "EmptyCloudDirectories", "This installation could not continue due to a configuration issue. Please contact support."));
+			bInstallerInitSuccess = false;
+
+			UE_LOG(LogBuildPatchServices, Error, TEXT("Installer configuration: Cloud Directories is empty."));
+			InstallerError->SetError(EBuildPatchInstallError::InitializationError, InitializationErrorCodes::EmptyCloudDirectories);
+			return false;
+		}
+
 		if (Configuration.bInstallToMemory && Configuration.bConstructFilesInMemory)
 		{
 			UE_LOG(LogBuildPatchServices, Error, TEXT("Installer configuration: Can't construct in memory and install to memory."));
@@ -722,45 +930,109 @@ namespace BuildPatchServices
 			}
 		}
 
+		// Make sure that we have all necessary secrets available unless we are in preload mode.
+		if (Configuration.InstallMode != EInstallMode::Preload)
+		{
+			TSet<FGuid> NecessarySecrets;
+			for (const FBuildPatchInstallerAction& InstallerAction : InstallerActions)
+			{
+				// Only test actions with install manifests.
+				// An uninstall action can complete without any encryption keys, and we don't need the keys for the current install if they were not needed for a new update.
+				const FBuildPatchAppManifest* InstallManifest = InstallerAction.TryGetInstallManifest();
+				if (InstallManifest && InstallManifest->GetFeatureLevel() >= EFeatureLevel::ChunkEncryptionSupport)
+				{
+					NecessarySecrets.Append(InstallManifest->GetNecessaryEncryptionSecretIds());
+				}
+			}
+			TSet<FGuid> AvailableSecrets;
+			TSet<FGuid> InvalidSecrets;
+			for (const TPair<FGuid, TArray<uint8>>& EncryptionSecret : Configuration.EncryptionSecrets)
+			{
+				AvailableSecrets.Add(EncryptionSecret.Key);
+				if (!Crypto->IsValidKey_AES_256_GCM(EncryptionSecret.Value))
+				{
+					InvalidSecrets.Add(EncryptionSecret.Key);
+				}
+			}
+			if (InvalidSecrets.Num() > 0)
+			{
+				TArray<FString> InvalidSecretsIds;
+				Algo::Transform(InvalidSecrets, InvalidSecretsIds, [](const FGuid& Id) { return Id.ToString(); });
+				UE_LOG(LogBuildPatchServices, Error, TEXT("Installer setup: Installation cannot complete, invalid encryption secret(s) provided for id(s) %s"), *FString::Join(InvalidSecretsIds, TEXT(", ")));
+				InstallerError->SetError(EBuildPatchInstallError::InitializationError, InitializationErrorCodes::InvalidEncryptionKeyForNonPreload, 0,
+					NSLOCTEXT("BuildPatchInstallError", "InvalidEncryptionKeyForNonPreload", "This installation could not continue as the required decryption data is not valid. Please contact support."));
+				bInstallerInitSuccess = false;
+			}
+			bool bMissingSecrets = NecessarySecrets.Difference(AvailableSecrets).Num() > 0;
+			if (bMissingSecrets)
+			{
+				TArray<FString> AvailableSecretsIds;
+				Algo::Transform(AvailableSecrets, AvailableSecretsIds, [](const FGuid& Id) { return Id.ToString(); });
+				UE_LOG(LogBuildPatchServices, Error, TEXT("Installer setup: Installation cannot complete without all necessary encryption secrets. Available %d of %d: %s"), AvailableSecretsIds.Num(), NecessarySecrets.Num(), *FString::Join(AvailableSecretsIds, TEXT(", ")));
+				InstallerError->SetError(EBuildPatchInstallError::InitializationError, InitializationErrorCodes::MissingEncryptionKeyForNonPreload, 0,
+					NSLOCTEXT("BuildPatchInstallError", "MissingEncryptionKeyForNonPreload", "This installation could not continue as the required decryption data is missing. Please contact support."));
+				bInstallerInitSuccess = false;
+			}
+		}
+
 		// do the delta optimization
-		typedef TTuple<FBuildPatchInstallerAction&, TUniquePtr<IOptimisedDelta>> FManifestInfoDeltaPair;
+		typedef TTuple<FBuildPatchInstallerAction&, TSharedRef<IOptimisedDelta>> FManifestInfoDeltaPair;
 		TArray<FManifestInfoDeltaPair> RunningOptimisedDeltas;
 		for (FBuildPatchInstallerAction& InstallerAction : InstallerActions)
 		{
 			if (InstallerAction.IsUpdate())
 			{
-				TUniquePtr<IOptimisedDelta> OptimisedDelta(FOptimisedDeltaFactory::Create(
-					ConfigHelpers::BuildOptimisedDeltaConfig(Configuration, InstallerAction), 
-					InstallerHelpers::BuildOptimisedDeltaDependencies(DownloadService)));
+				TArray<FString> CloudDirectories = Configuration.CloudDirectories;
+				if (!InstallerAction.GetCloudSubdirectory().IsEmpty())
+				{
+					for (FString& CloudDirectory : CloudDirectories)
+					{
+						CloudDirectory /= InstallerAction.GetCloudSubdirectory();
+					}
+				}
+				const int32_t NumberOfDirectories = CloudDirectories.Num();
+
+				TSharedRef<IUriProvider, ESPMode::ThreadSafe> UriProvider = FUriProviderFactory::Create(MessagePump.Get(), MoveTemp(CloudDirectories));
+				FOptimisedDeltaDependencies OptimisedDeltaDependencies = InstallerHelpers::BuildOptimisedDeltaDependencies(DownloadService, MoveTemp(UriProvider));
+				TSharedRef<IOptimisedDelta> OptimisedDelta = FOptimisedDeltaFactory::Create(ConfigHelpers::BuildOptimisedDeltaConfig(Configuration, InstallerAction, NumberOfDirectories), MoveTemp(OptimisedDeltaDependencies));
 				RunningOptimisedDeltas.Add(FManifestInfoDeltaPair{ InstallerAction, MoveTemp(OptimisedDelta) });
 			}
 		}
 		for (FManifestInfoDeltaPair& RunningOptimisedDelta : RunningOptimisedDeltas)
 		{
 			FBuildPatchInstallerAction& InstallerAction = RunningOptimisedDelta.Get<0>();
-			const TUniquePtr<IOptimisedDelta>& OptimisedDelta = RunningOptimisedDelta.Get<1>();
-			const IOptimisedDelta::FResultValueOrError& OptimisedDeltaResult = OptimisedDelta->GetResult();
-			PreviousTotalDownloadRequired.Add(OptimisedDelta->GetMetaDownloadSize());
-			// The OptimiseDelta class handles policy, so if we get a nullptr back, that is a hard error.
-			if (!OptimisedDeltaResult.IsValid())
-			{
-				UE_LOG(LogBuildPatchServices, Error, TEXT("Installer setup: Destination manifest could not be obtained."));
-				InstallerError->SetError(EBuildPatchInstallError::DownloadError, *OptimisedDeltaResult.GetError());
-				bInstallerInitSuccess = false;
-				break;
+			TSharedRef<IOptimisedDelta> OptimisedDelta = RunningOptimisedDelta.Get<1>();
+			if (WaitForResultOrAbort(*OptimisedDelta, bShouldAbort))
+			{
+				const IOptimisedDelta::FResultValueOrError& OptimisedDeltaResult = OptimisedDelta->GetResult();
+				PreviousTotalDownloadRequired += OptimisedDelta->GetMetaDownloadSize();
+				// The OptimiseDelta class handles policy, so if we get a nullptr back, that is a hard error.
+				if (!OptimisedDeltaResult.IsValid())
+				{
+					UE_LOG(LogBuildPatchServices, Error, TEXT("Installer setup: Destination manifest could not be obtained."));
+					InstallerError->SetError(EBuildPatchInstallError::DownloadError, *OptimisedDeltaResult.GetError());
+					bInstallerInitSuccess = false;
+					break;
+				}
+				else
+				{
+					InstallerAction.SetDeltaManifest(OptimisedDeltaResult.GetValue().ToSharedRef());
+				}
 			}
 			else
 			{
-				InstallerAction.SetDeltaManifest(OptimisedDeltaResult.GetValue().ToSharedRef());
+				UE_LOG(LogBuildPatchServices, Error, TEXT("Installer setup is aborted."));
+				InstallerError->SetError(EBuildPatchInstallError::UserCanceled, UserCancelErrorCodes::UserRequested);
+				bInstallerInitSuccess = false;
+				break;
 			}
 		}
-		RunningOptimisedDeltas.Empty();
 
 		// We can now build out any systems that need late construction but can survive between retries.
 		ManifestSet.Reset(FBuildManifestSetFactory::Create(InstallerActions));
 		Verifier.Reset(FVerifierFactory::Create(FileSystem.Get(), VerifierStatistics.Get(), Configuration.VerifyMode, Configuration.SharedContext, 
-			ManifestSet.Get(), Configuration.InstallDirectory, Configuration.InstallMode == EInstallMode::StageFiles ? InstallStagingDir : FString(),
-			NewPerFileSubdirectories));
+			ManifestSet.Get(), Configuration.InstallDirectory, Configuration.InstallMode == EInstallMode::StageFiles ? InstallStagingDir : FString(), 
+			!Configuration.bStageWithRawFilenames, NewPerFileSubdirectories));
 
 		// Add systems to controllables.
 		{
@@ -768,15 +1040,13 @@ namespace BuildPatchServices
 			Controllables.Add(Verifier.Get());
 		}
 
-		// Update to chunk data size cache since we may have delta manifests now
+		// Queue update to chunk data size cache on main thread
+		if (bInstallerInitSuccess)
 		{
-			TArray<FBuildPatchAppManifestPtr> ActionManifests;
-			ActionManifests.Reserve(InstallerActions.Num());
 			for (const FBuildPatchInstallerAction& InstallerAction : InstallerActions)
 			{
-				ActionManifests.Emplace(InstallerAction.TryGetSharedInstallManifest());
+				ChunkDataSizeProvider->AddManifestData(InstallerAction.TryGetSharedInstallManifest());
 			}
-			ChunkDataSizeProvider->AddManifestData(ActionManifests);
 		}
 
 		// Init build statistics that are known.
@@ -832,6 +1102,15 @@ namespace BuildPatchServices
 		// No longer queued
 		BuildProgress.SetStateProgress(EBuildPatchState::Queued, 1.0f);
 
+		// If we have resuming data, then we immediately go to Resuming
+		const FString LegacyResumeDataFilename = InstallStagingDir / TEXT("$resumeData");
+		const FString ResumeDataFilename = MetaStagingDir / TEXT("$resumeData");
+		const bool bHasResumeData = FileSystem->FileExists(*ResumeDataFilename) || FileSystem->FileExists(*LegacyResumeDataFilename);
+		if (bHasResumeData)
+		{
+			BuildProgress.SetStateProgress(EBuildPatchState::Initializing, 1.0f);
+		}
+
 		// Init prereqs progress value
 		bool bHasPrereqPath = false;
 		for (const FBuildPatchInstallerAction& InstallerAction : InstallerActions)
@@ -851,10 +1130,10 @@ namespace BuildPatchServices
 
 		uint64 MaxDiskSpaceNeededWhenDeletingChunkDbsIfRequested = 0;
 		if (Configuration.bCalculateDeleteChunkDbMaxDiskSpaceAndExit)
-		{		
+		{
 			if (bProcessSuccess)
 			{
-				MaxDiskSpaceNeededWhenDeletingChunkDbsIfRequested = DetermineInstallMaxDiskSizeIfDeletingChunkDbs(Configuration.ChunkDatabaseFiles, ManifestSet.Get(), FileSystem.Get(), Configuration.InstallDirectory, {}, Configuration.InstallMode == EInstallMode::PrereqOnly);
+				MaxDiskSpaceNeededWhenDeletingChunkDbsIfRequested = DetermineInstallMaxDiskSizeIfDeletingChunkDbs(Configuration.ChunkDatabaseFiles, ManifestSet.Get(), FileSystem.Get(), Configuration.InstallDirectory, {}, Configuration.InstallMode);
 			}
 		}
 		else
@@ -867,37 +1146,51 @@ namespace BuildPatchServices
 
 				// Keep retrying the install while it is not canceled, or caused by download error
 				bProcessSuccess = false;
+				bool bInstallSuccess = false;
 				bool bCanRetry = true;
+				bool bWillRetry = false;
 				int32 InstallRetries = ConfigHelpers::NumInstallerRetries();
 				while (!bProcessSuccess && bCanRetry)
 				{
+					bLastInstallIteration = 1 == InstallRetries;
 					// Inform file operation tracker of the selected manifest.
 					FileOperationTracker->OnManifestSelection(ManifestSet.Get());
 
-					// Run the install
-					bool bInstallSuccess = RunInstallation(CorruptFiles);
+					if (Configuration.InstallMode == EInstallMode::Preload)
+					{
+						// Run the install
+						bInstallSuccess = bProcessSuccess = RunPreload();
+					}
+					else
+					{
+						// Run the install
+						bInstallSuccess = RunInstallation(CorruptFiles);
+
+						// Runs the custom uninstall actions if specified
+						bInstallSuccess = bInstallSuccess && RunCustomUninstallAction();
 
-					// Backup local changes then move generated files
-					bInstallSuccess = bInstallSuccess && RunBackupAndMove();
+						// Backup local changes then move generated files
+						bInstallSuccess = bInstallSuccess && RunBackupAndMove();
 
-					// Setup file attributes
-					bInstallSuccess = bInstallSuccess && RunFileAttributes();
+						// Setup file attributes
+						bInstallSuccess = bInstallSuccess && RunFileAttributes();
 
-					// Run Verification
-					CorruptFiles.Empty();
-					bProcessSuccess = bInstallSuccess && RunVerification(CorruptFiles);
+						// Run Verification
+						CorruptFiles.Empty();
+						bProcessSuccess = bInstallSuccess && RunVerification(CorruptFiles);
+					}
 
 					// Set if we can retry
 					--InstallRetries;
-					bCanRetry = InstallRetries > 0 && !InstallerError->IsCancelled() && InstallerError->CanRetry();
-					const bool bWillRetry = !bProcessSuccess && bCanRetry;
+					bCanRetry = InstallRetries > 0 && !InstallerError->IsCancelled() && InstallerError->CanRetry() && !ManifestSet->IsExclusivelyVerifying();
+					bWillRetry = !bProcessSuccess && bCanRetry;
 
 					// Clean staging if INSTALL success, we still do cleanup if we failed at the verify stage.
-					if (bInstallSuccess && 
-						!Configuration.bInstallToMemory)
+					const bool bShouldClean = bInstallSuccess && !Configuration.bInstallToMemory && Configuration.InstallMode != EInstallMode::Preload;
+					if (bShouldClean)
 					{
 						CleanUpTimer.Start();
-						if (Configuration.InstallMode == EInstallMode::StageFiles)
+						if (Configuration.InstallMode == EInstallMode::StageFiles || bAreResultsEncrypted)
 						{
 							UE_LOG(LogBuildPatchServices, Log, TEXT("Deleting litter from staging area."));
 							IFileManager::Get().DeleteDirectory(*DataStagingDir, false, true);
@@ -930,8 +1223,10 @@ namespace BuildPatchServices
 					float TempFinalProgress = BuildProgress.GetProgressNoMarquee();
 					{
 						FScopeLock Lock(&ThreadLock);
-						BuildStats.NumInstallRetries = ConfigHelpers::NumInstallerRetries() - (InstallRetries + 1);
+						const uint32 NumInstallRetries = ConfigHelpers::NumInstallerRetries() - (InstallRetries + 1);
+						BuildStats.NumInstallRetries = NumInstallRetries;
 						BuildStats.FinalProgress = TempFinalProgress;
+						BuildStats.NumInstallRetries = NumInstallRetries;
 						// If we failed, and will retry, record this failure type and reset the abort flag
 						if (bWillRetry)
 						{
@@ -945,7 +1240,8 @@ namespace BuildPatchServices
 					if (bWillRetry)
 					{
 						InitializeTimer.Start();
-						BuildProgress.CancelAbort();
+						Verifier->Reset();
+						BuildProgress.Reset();
 						BuildProgress.SetStateProgress(EBuildPatchState::Initializing, 0.0f);
 						BuildProgress.SetStateProgress(EBuildPatchState::Resuming, 0.0f);
 						BuildProgress.SetStateProgress(EBuildPatchState::Downloading, 0.0f);
@@ -973,6 +1269,7 @@ namespace BuildPatchServices
 		// Make sure all timers are stopped
 		InitializeTimer.Stop();
 		ConstructTimer.Stop();
+		UninstallActionTimer.Stop();
 		MoveFromStageTimer.Stop();
 		FileAttributesTimer.Stop();
 		VerifyTimer.Stop();
@@ -988,6 +1285,7 @@ namespace BuildPatchServices
 			FScopeLock Lock(&ThreadLock);
 			BuildStats.InitializeTime = InitializeTimer.GetSeconds();
 			BuildStats.ConstructTime = ConstructTimer.GetSeconds();
+			BuildStats.UninstallActionTime = UninstallActionTimer.GetSeconds();
 			BuildStats.MoveFromStageTime = MoveFromStageTimer.GetSeconds();
 			BuildStats.FileAttributesTime = FileAttributesTimer.GetSeconds();
 			BuildStats.VerifyTime = VerifyTimer.GetSeconds();
@@ -1038,14 +1336,17 @@ namespace BuildPatchServices
 		return false;
 	}
 
-	static FChunkDbSourceConfig BuildChunkDbSourceConfig(const FBuildInstallerConfiguration& InConfiguration)
+	static FChunkDbSourceConfig BuildChunkDbSourceConfig(const FBuildInstallerConfiguration& Configuration, const FString& DataStagingDir)
 	{
-		FChunkDbSourceConfig ChunkDbSourceConfig(InConfiguration.ChunkDatabaseFiles);
-		ChunkDbSourceConfig.bDeleteChunkDBAfterUse = InConfiguration.bDeleteChunkDbFilesAfterUse;
+		TArray<FString> ChunkDatabaseFiles = Configuration.ChunkDatabaseFiles;
+		Algo::Transform(ConfigHelpers::EnumeratePreloadChunkDbs(*DataStagingDir), ChunkDatabaseFiles, [&DataStagingDir](const FString& ChunkDbFile) {return DataStagingDir / ChunkDbFile; });
+		FChunkDbSourceConfig ChunkDbSourceConfig(MoveTemp(ChunkDatabaseFiles));
+		ChunkDbSourceConfig.bDeleteChunkDBAfterUse = Configuration.bDeleteChunkDbFilesAfterUse;
+		ChunkDbSourceConfig.bTruncateChunkDbs = Configuration.bTruncateChunkDbFilesAsUsed;
 		return ChunkDbSourceConfig;
 	}
 
-	static FConstructorCloudChunkSourceConfig BuildCloudSourceConfig(const FBuildInstallerConfiguration& InConfiguration)
+	static FConstructorCloudChunkSourceConfig BuildConstructorCloudSourceConfig(const FBuildInstallerConfiguration& InConfiguration)
 	{
 		FConstructorCloudChunkSourceConfig CloudSourceConfig(InConfiguration.CloudDirectories);
 
@@ -1084,8 +1385,8 @@ namespace BuildPatchServices
 		TArray<float> HealthPercentages;
 		HealthPercentages.AddZeroed((int32)EBuildPatchDownloadHealth::NUM_Values);
 		if (GConfig->GetFloat(TEXT("Portal.BuildPatch"), TEXT("OKHealth"), HealthPercentages[(int32)EBuildPatchDownloadHealth::OK], GEngineIni)
-		 && GConfig->GetFloat(TEXT("Portal.BuildPatch"), TEXT("GoodHealth"), HealthPercentages[(int32)EBuildPatchDownloadHealth::Good], GEngineIni)
-		 && GConfig->GetFloat(TEXT("Portal.BuildPatch"), TEXT("ExcellentHealth"), HealthPercentages[(int32)EBuildPatchDownloadHealth::Excellent], GEngineIni))
+			&& GConfig->GetFloat(TEXT("Portal.BuildPatch"), TEXT("GoodHealth"), HealthPercentages[(int32)EBuildPatchDownloadHealth::Good], GEngineIni)
+			&& GConfig->GetFloat(TEXT("Portal.BuildPatch"), TEXT("ExcellentHealth"), HealthPercentages[(int32)EBuildPatchDownloadHealth::Excellent], GEngineIni))
 		{
 			CloudSourceConfig.HealthPercentages = MoveTemp(HealthPercentages);
 		}
@@ -1094,27 +1395,91 @@ namespace BuildPatchServices
 		GConfig->GetFloat(TEXT("Portal.BuildPatch"), TEXT("DisconnectedDelay"), CloudSourceConfig.DisconnectedDelay, GEngineIni);
 		CloudSourceConfig.DisconnectedDelay = FMath::Clamp<float>(CloudSourceConfig.DisconnectedDelay, 1.0f, 30.0f);
 
+		CloudSourceConfig.DownloadRateLimitBps = InConfiguration.DownloadRateLimitBps;
+
 		return CloudSourceConfig;
 	}
 
-	FDownloadConnectionCountConfig FBuildPatchInstaller::BuildConnectionCountConfig()
+	FCloudSourceConfig FBuildPatchInstaller::BuildCloudSourceConfig()
 	{
-		FDownloadConnectionCountConfig ConnectionCountConfiguration;
+		FCloudSourceConfig CloudSourceConfig(Configuration.CloudDirectories);
+		CloudSourceConfig.SharedContext = Configuration.SharedContext.Get();
+		CloudSourceConfig.DownloadRateLimitBps = Configuration.DownloadRateLimitBps;
 
-		// Load simultaneous downloads from engine config.
-		GConfig->GetInt(TEXT("Portal.BuildPatch"), TEXT("ChunkDownloads"), (int32&)ConnectionCountConfiguration.FallbackCount, GEngineIni);
-		ConnectionCountConfiguration.FallbackCount = FMath::Clamp<uint32>(ConnectionCountConfiguration.FallbackCount, 1, 100);
+		// Load max download retry count from engine config.
+		GConfig->GetInt(TEXT("Portal.BuildPatch"), TEXT("ChunkRetries"), CloudSourceConfig.MaxRetryCount, GEngineIni);
+		CloudSourceConfig.MaxRetryCount = FMath::Clamp<int32>(CloudSourceConfig.MaxRetryCount, -1, 1000);
 
-		bool bDisableDownloadConnectionScaling = false;
-		if (GConfig->GetBool(TEXT("Portal.BuildPatch"), TEXT("ChunkDownloadsDisableConnectionScaling"), bDisableDownloadConnectionScaling, GEngineIni))
+		// Load the number of retries before enabling detailed logging.
+		GConfig->GetInt(TEXT("Portal.BuildPatch"), TEXT("RetriesToEnableVerbose"), CloudSourceConfig.RetriesToEnableVerbose, GEngineIni);
+		CloudSourceConfig.RetriesToEnableVerbose = FMath::Clamp<int32>(CloudSourceConfig.RetriesToEnableVerbose, 0, 1000);
+
+		// Load prefetch config
+		GConfig->GetInt(TEXT("Portal.BuildPatch"), TEXT("CloudSourcePreFetchMinimum"), CloudSourceConfig.PreFetchMinimum, GEngineIni);
+		GConfig->GetInt(TEXT("Portal.BuildPatch"), TEXT("CloudSourcePreFetchMaximum"), CloudSourceConfig.PreFetchMaximum, GEngineIni);
+		CloudSourceConfig.PreFetchMinimum = FMath::Clamp<int32>(CloudSourceConfig.PreFetchMinimum, 1, 1000);
+		CloudSourceConfig.PreFetchMaximum = FMath::Clamp<int32>(CloudSourceConfig.PreFetchMaximum, CloudSourceConfig.PreFetchMinimum, 1000);
+
+		// Load retry times from engine config.
+		TArray<FString> ConfigStrings;
+		GConfig->GetArray(TEXT("Portal.BuildPatch"), TEXT("RetryTimes"), ConfigStrings, GEngineIni);
+		bool bReadArraySuccess = ConfigStrings.Num() > 0;
+		TArray<float> RetryDelayTimes;
+		RetryDelayTimes.AddZeroed(ConfigStrings.Num());
+		for (int32 TimeIdx = 0; TimeIdx < ConfigStrings.Num() && bReadArraySuccess; ++TimeIdx)
 		{
-			if (bDisableDownloadConnectionScaling)
+			float TimeValue = FPlatformString::Atof(*ConfigStrings[TimeIdx]);
+			// Atof will return 0.0 if failed to parse, and we don't expect a time of 0.0 so presume error
+			if (TimeValue > 0.0f)
 			{
-				ConnectionCountConfiguration.bDisableConnectionScaling = true;
+				RetryDelayTimes[TimeIdx] = FMath::Clamp<float>(TimeValue, 0.5f, 300.0f);
 			}
-		}
-
-		uint32 MaxLibCurlConnections = 0;
+			else
+			{
+				bReadArraySuccess = false;
+			}
+		}
+		// If the retry array was parsed successfully, set on config.
+		if (bReadArraySuccess)
+		{
+			CloudSourceConfig.RetryDelayTimes = MoveTemp(RetryDelayTimes);
+		}
+
+		// Load percentiles for download health groupings from engine config.
+		// If the enum was changed since writing, the config here needs updating.
+		check((int32)EBuildPatchDownloadHealth::NUM_Values == 5);
+		TArray<float> HealthPercentages;
+		HealthPercentages.AddZeroed((int32)EBuildPatchDownloadHealth::NUM_Values);
+		if (GConfig->GetFloat(TEXT("Portal.BuildPatch"), TEXT("OKHealth"), HealthPercentages[(int32)EBuildPatchDownloadHealth::OK], GEngineIni)
+			&& GConfig->GetFloat(TEXT("Portal.BuildPatch"), TEXT("GoodHealth"), HealthPercentages[(int32)EBuildPatchDownloadHealth::Good], GEngineIni)
+			&& GConfig->GetFloat(TEXT("Portal.BuildPatch"), TEXT("ExcellentHealth"), HealthPercentages[(int32)EBuildPatchDownloadHealth::Excellent], GEngineIni))
+		{
+			CloudSourceConfig.HealthPercentages = MoveTemp(HealthPercentages);
+		}
+
+		// Load the delay for how long we get no data for until determining the health as disconnected.
+		GConfig->GetFloat(TEXT("Portal.BuildPatch"), TEXT("DisconnectedDelay"), CloudSourceConfig.DisconnectedDelay, GEngineIni);
+		CloudSourceConfig.DisconnectedDelay = FMath::Clamp<float>(CloudSourceConfig.DisconnectedDelay, 1.0f, 30.0f);
+
+		// We tell the cloud source to only start downloads once it receives the first get call.
+		CloudSourceConfig.bBeginDownloadsOnFirstGet = true;
+
+		return CloudSourceConfig;
+	}
+
+	FDownloadConnectionCountConfig FBuildPatchInstaller::BuildConnectionCountConfig()
+	{
+		FDownloadConnectionCountConfig ConnectionCountConfiguration;
+
+		// Load simultaneous downloads from engine config.
+		GConfig->GetInt(TEXT("Portal.BuildPatch"), TEXT("ChunkDownloads"), (int32&)ConnectionCountConfiguration.FallbackCount, GEngineIni);
+		ConnectionCountConfiguration.FallbackCount = FMath::Clamp<uint32>(ConnectionCountConfiguration.FallbackCount, 1, 100);
+
+		// Check if download connection scaling is disabled via config or download throttling.
+		GConfig->GetBool(TEXT("Portal.BuildPatch"), TEXT("ChunkDownloadsDisableConnectionScaling"), ConnectionCountConfiguration.bDisableConnectionScaling, GEngineIni);
+		ConnectionCountConfiguration.bDisableConnectionScaling = ConnectionCountConfiguration.bDisableConnectionScaling || DownloadRateLimitBps != 0;
+
+		uint32 MaxLibCurlConnections = 0;
 		uint32 MaxDownloadCount = ConnectionCountConfiguration.MaxLimit;
 		if (!GConfig->GetInt(TEXT("HTTP"), TEXT("HttpMaxConnectionsPerServer"), (int32&)MaxLibCurlConnections, GEngineIni))
 		{
@@ -1133,8 +1498,17 @@ namespace BuildPatchServices
 		uint32 MinLimit = ConnectionCountConfiguration.MinLimit;
 		GConfig->GetInt(TEXT("Portal.BuildPatch"), TEXT("ChunkDownloadsLowerLimit"), (int32&)ConnectionCountConfiguration.MinLimit, GEngineIni);
 		ConnectionCountConfiguration.MinLimit = FMath::Clamp<uint32>(ConnectionCountConfiguration.MinLimit,
-																		FMath::Min(MinLimit, MaxDownloadCount - 2),
-																		FMath::Min(32U, MaxDownloadCount - 2));
+			FMath::Min(MinLimit, MaxDownloadCount - 2),
+			FMath::Min(32U, MaxDownloadCount - 2));
+
+		// Overriding MinLimit for A/B Testing
+		const FString UserId = FPlatformMisc::GetEpicAccountId();
+		const FString Variant = Experiments::GetVariant(Experiments::DownloadSpeed(), UserId);
+
+		if (Variant == TEXT("Connection_Floor_16"))
+		{
+			ConnectionCountConfiguration.MinLimit = 16;
+		}
 
 		GConfig->GetInt(TEXT("Portal.BuildPatch"), TEXT("ChunkDownloadsUpperLimit"), (int32&)ConnectionCountConfiguration.MaxLimit, GEngineIni);
 		ConnectionCountConfiguration.MaxLimit = FMath::Clamp<uint32>(ConnectionCountConfiguration.MaxLimit, ConnectionCountConfiguration.MinLimit + 2, MaxDownloadCount);
@@ -1160,6 +1534,13 @@ namespace BuildPatchServices
 		return ConnectionCountConfiguration;
 	}
 
+	FString FBuildPatchInstaller::GetSHA1Filename(const FFileManifest& FileManifest)
+	{
+		FString FileHashStringValue;
+		FBuildPatchUtils::SHAToBase32(FileManifest.SHA1Hash, FileHashStringValue);
+
+		return FileHashStringValue;
+	}
 
 	static void GenerateFilesToConstruct(TSet<FString>& FilesToConstruct, IBuildManifestSet* ManifestSet, const TArray<FString>& CorruptFiles, const TSet<FString>& TaggedFiles, const TSet<FString>& OutdatedFiles, bool bIsPrereqOnly)
 	{
@@ -1196,7 +1577,7 @@ namespace BuildPatchServices
 	// that no new chunks are required, the first file is completely deleted (no staging size), and the
 	// staging size of further files fits in that freed up space.
 	//
-	static uint64 DetermineInstallMaxDiskSizeIfDeletingChunkDbs(const TArray<FString>& ChunkDbFiles, IBuildManifestSet* ManifestSet, IFileSystem* FileSystem, const FString& InstallDirectory, const TArray<FString>& CorruptFiles, bool bIsPrereqOnly)
+	static uint64 DetermineInstallMaxDiskSizeIfDeletingChunkDbs(const TArray<FString>& ChunkDbFiles, IBuildManifestSet* ManifestSet, IFileSystem* FileSystem, const FString& InstallDirectory, const TArray<FString>& CorruptFiles, EInstallMode InstallMode)
 	{
 		TSet<FString> FilesToConstructSet;
 
@@ -1205,7 +1586,7 @@ namespace BuildPatchServices
 		TSet<FString> OutdatedFiles;
 		ManifestSet->GetOutdatedFiles(InstallDirectory, OutdatedFiles);
 
-		GenerateFilesToConstruct(FilesToConstructSet, ManifestSet, CorruptFiles, TaggedFiles, OutdatedFiles, bIsPrereqOnly);
+		GenerateFilesToConstruct(FilesToConstructSet, ManifestSet, CorruptFiles, TaggedFiles, OutdatedFiles, InstallMode == EInstallMode::PrereqOnly);
 
 		// Generate the list of chunks we will consume, in order.
 		TUniquePtr<IChunkReferenceTracker> ChunkReferenceTracker(FChunkReferenceTrackerFactory::Create(
@@ -1257,11 +1638,11 @@ namespace BuildPatchServices
 		}
 
 		TArray<uint64> ChunkDbSizesAtPosition;
-		uint64 TotalChunkDbSize = IConstructorChunkDbChunkSource::GetChunkDbSizesAtIndexes(ChunkDbFiles, FileSystem, ReferenceChain, FileCompletionPositions, ChunkDbSizesAtPosition);
+		IConstructorChunkDbChunkSource::GetChunkDbSizesAtIndexes(ChunkDbFiles, FileSystem, ReferenceChain, FileCompletionPositions, ChunkDbSizesAtPosition);
 
 		uint64 MaxDiskSize = FBuildPatchUtils::CalculateDiskSpaceRequirementsWithDeleteDuringInstall(
-			FilesToConstruct, 0, ManifestSet,
-			ChunkDbSizesAtPosition, TotalChunkDbSize);
+			FilesToConstruct, 0, 0, ManifestSet,
+			ChunkDbSizesAtPosition, InstallDirectory, InstallMode);
 
 		return MaxDiskSize;
 	}
@@ -1269,19 +1650,6 @@ namespace BuildPatchServices
 
 	bool FBuildPatchInstaller::RunInstallation(TArray<FString>& CorruptFiles)
 	{
-		UE_LOG(LogBuildPatchServices, Log, TEXT("Starting Installation"));
-		// Save the staging directories
-		FPaths::NormalizeDirectoryName(DataStagingDir);
-		FPaths::NormalizeDirectoryName(InstallStagingDir);
-		FPaths::NormalizeDirectoryName(MetaStagingDir);
-
-		// UE file system abstractions create all necessary directories when creating files.
-		// By not doing this up front we avoid any unnecessary file operations as some platforms are sensitive to that.
-		//IFileManager::Get().MakeDirectory(*InstallStagingDir, true);
-		//IFileManager::Get().MakeDirectory(*DataStagingDir, true);
-		//IFileManager::Get().MakeDirectory(*MetaStagingDir, true);
-
-		
 		// We want to reset errors because we are retrying, however we ONLY want to do this if the error was caused by us!
 		// This is a bit delicate because this pointer gets hit by any thread, but we know that installer errors are latching.
 		// (except for retry).
@@ -1291,6 +1659,29 @@ namespace BuildPatchServices
 			InstallerError->RegisterForErrors([this]() { CancelInstall(); });
 		}
 
+		// Installation is not ran for VerifyOnly actions. This mode is read-only.
+		if (ManifestSet->IsExclusivelyVerifying())
+		{
+			UE_LOG(LogBuildPatchServices, Log, TEXT("Skipping installation for VerifyOnly action"));
+			// Set weights for verify only
+			BuildProgress.SetStateWeight(EBuildPatchState::Downloading, 0.0f);
+			BuildProgress.SetStateWeight(EBuildPatchState::Installing, 0.0f);
+			BuildProgress.SetStateWeight(EBuildPatchState::MovingToInstall, 0.0f);
+			BuildProgress.SetStateWeight(EBuildPatchState::SettingAttributes, 0.0f);
+			BuildProgress.SetStateWeight(EBuildPatchState::BuildVerification, 1.0f);
+			// Mark all installation steps complete
+			BuildProgress.SetStateProgress(EBuildPatchState::Initializing, 1.0f);
+			BuildProgress.SetStateProgress(EBuildPatchState::Resuming, 1.0f);
+			BuildProgress.SetStateProgress(EBuildPatchState::Downloading, 1.0f);
+			BuildProgress.SetStateProgress(EBuildPatchState::Installing, 1.0f);
+			BuildProgress.SetStateProgress(EBuildPatchState::MovingToInstall, 1.0f);
+			BuildProgress.SetStateProgress(EBuildPatchState::SettingAttributes, 1.0f);
+			// Stop relevant timers
+			InitializeTimer.Stop();
+			return true;
+		}
+		UE_LOG(LogBuildPatchServices, Log, TEXT("Starting Installation"));
+
 		// Get the list of required files, by the tags
 		TaggedFiles.Empty();
 		ManifestSet->GetExpectedFiles(TaggedFiles);
@@ -1351,21 +1742,38 @@ namespace BuildPatchServices
 		// Make sure all the files won't exceed the maximum path length
 		if (!Configuration.bInstallToMemory)
 		{
+			{
+				// +1 means extra char for path delimiter
+				int32 StagedConstructionFilePathLen = InstallStagingDir.Len() + 1 + BuildPatchConstants::SHAToBase32StringLength;
+
+				if (StagedConstructionFilePathLen >= FPlatformMisc::GetMaxPathLength())
+				{
+					UE_LOG(LogBuildPatchServices, Error, TEXT("Staged filenames (length %d) will exceed maximum path length when saved to %s"), BuildPatchConstants::SHAToBase32StringLength, *InstallStagingDir);
+					InstallerError->SetError(EBuildPatchInstallError::PathLengthExceeded, PathLengthErrorCodes::StagingDirectory);
+					return false;
+				}
+			}
 			for (const FString& FileToConstruct : FilesToConstruct)
 			{
 				const FString InstallConstructionFile = FBuildPatchUtils::ResolveInstallationFileName(Configuration.InstallDirectory, NewPerFileSubdirectories, FileToConstruct);
-				const FString StagedConstructionFile = InstallStagingDir / FileToConstruct;
+
 				if (InstallConstructionFile.Len() >= FPlatformMisc::GetMaxPathLength())
 				{
 					UE_LOG(LogBuildPatchServices, Error, TEXT("Could not create new file due to exceeding maximum path length %s"), *InstallConstructionFile);
 					InstallerError->SetError(EBuildPatchInstallError::PathLengthExceeded, PathLengthErrorCodes::InstallDirectory);
 					return false;
 				}
-				if (StagedConstructionFile.Len() >= FPlatformMisc::GetMaxPathLength())
+
+				if (Configuration.bStageWithRawFilenames)
 				{
-					UE_LOG(LogBuildPatchServices, Error, TEXT("Could not create new file due to exceeding maximum path length %s"), *StagedConstructionFile);
-					InstallerError->SetError(EBuildPatchInstallError::PathLengthExceeded, PathLengthErrorCodes::StagingDirectory);
-					return false;
+					FString StagedConstructionFile = InstallStagingDir / FileToConstruct;
+
+					if (StagedConstructionFile.Len() >= FPlatformMisc::GetMaxPathLength())
+					{
+						UE_LOG(LogBuildPatchServices, Error, TEXT("Could not create new file due to exceeding maximum path length %s"), *StagedConstructionFile);
+						InstallerError->SetError(EBuildPatchInstallError::PathLengthExceeded, PathLengthErrorCodes::StagingDirectory);
+						return false;
+					}
 				}
 			}
 		}
@@ -1382,27 +1790,46 @@ namespace BuildPatchServices
 		}
 
 		// Cache the last download requirement in case we are running a retry.
-		PreviousTotalDownloadRequired.Add(CloudChunkSourceStatistics->GetRequiredDownloadSize());
+		PreviousTotalDownloadRequired += CloudChunkSourceStatistics->GetDownloadedBytes();
 		// Reset so that we don't double count data.
 		CloudChunkSourceStatistics->OnRequiredDataUpdated(0);
 		CloudChunkSourceStatistics->OnReceivedDataUpdated(0);
 
+		// De-dupe the file construction list if we are using SHA1 filenames for staging. In this case we don't build duplicate files, we copy them.
+		TSet<FString> DeDupedFilesToConstruct;
+		if (!Configuration.bStageWithRawFilenames)
+		{
+			TSet<FSHAHash> FileShaReferences;
+			for (const FString& File : FilesToConstruct)
+			{
+				const FFileManifest* NewFileManifest = ManifestSet->GetNewFileManifest(File);
+				if (NewFileManifest != nullptr && !FileShaReferences.Contains(NewFileManifest->SHA1Hash))
+				{
+					FileShaReferences.Add(NewFileManifest->SHA1Hash);
+					DeDupedFilesToConstruct.Add(File);
+				}
+			}
+		}
+
+		// Before beginning work, make sure we have deleted the HDD reservation file if it existed.
+		FileSystem->DeleteFile(*HddSpaceReservationFile);
+
 		// Scoped systems composition and execution.
 		{
 			TUniquePtr<IChunkDataSerialization> ChunkDataSerialization(FChunkDataSerializationFactory::Create(
-				FileSystem.Get()));
+				FileSystem.Get(), Crypto.Get(), { EFeatureLevel::Latest, Configuration.EncryptionSecrets, FGuid() /* UE5 MERGE TODO : , CloudChunkSourceStatistics.Get()*/}));
 
 			// Generate the list of chunks we will need and the order in which we will need them.
 			TUniquePtr<IChunkReferenceTracker> ChunkReferenceTracker(FChunkReferenceTrackerFactory::Create(
 				ManifestSet.Get(),
-				FilesToConstruct));
+				Configuration.bStageWithRawFilenames ? FilesToConstruct : DeDupedFilesToConstruct));
 			TSet<FGuid> ReferencedChunks = ChunkReferenceTracker->GetReferencedChunks();
 			
 			// Add a source for pulling from chunk "databases", basically tarballs of chunks.
 			TArray<FGuid> ChunkAccessOrderedList;
 			ChunkReferenceTracker->CopyOutOrderedUseList(ChunkAccessOrderedList);
 			TUniquePtr<IConstructorChunkDbChunkSource> ChunkDbChunkSource(IConstructorChunkDbChunkSource::CreateChunkDbSource(
-				BuildChunkDbSourceConfig(Configuration),
+				BuildChunkDbSourceConfig(Configuration, DataStagingDir),
 				FileSystem.Get(),
 				ChunkAccessOrderedList,
 				ChunkDataSerialization.Get(),
@@ -1420,6 +1847,7 @@ namespace BuildPatchServices
 					InstallChunkSourceStatistics.Get(),
 					InstallationInfo,
 					ReferencedChunks,
+					false,
 					OldPerFileSubdirectories
 				));
 			}
@@ -1472,7 +1900,7 @@ namespace BuildPatchServices
 
 			TUniquePtr<IConstructorCloudChunkSource> CloudChunkSource(
 				IConstructorCloudChunkSource::CreateCloudSource(
-					BuildCloudSourceConfig(Configuration),
+					BuildConstructorCloudSourceConfig(Configuration),
 					DownloadService.Get(),
 					ChunkDataSerialization.Get(),
 					DownloadConnectionCount.Get(),
@@ -1486,7 +1914,7 @@ namespace BuildPatchServices
 			FCC.bInstallToMemory = Configuration.bInstallToMemory;
 			FCC.bConstructInMemory = Configuration.bConstructFilesInMemory;
 			FCC.bSkipInitialDiskSizeCheck = Configuration.bSkipInitialDiskSizeCheck;
-			FCC.ConstructList = FilesToConstruct.Array();
+			FCC.ConstructList = (!Configuration.bStageWithRawFilenames ? DeDupedFilesToConstruct : FilesToConstruct).Array();
 			FCC.InstallDirectory = Configuration.InstallDirectory;
 			FPaths::NormalizeDirectoryName(FCC.InstallDirectory);
 			FCC.InstallMode = Configuration.InstallMode;
@@ -1502,7 +1930,9 @@ namespace BuildPatchServices
 			FCC.IOBatchSizeMB = Configuration.ConstructorIOBatchSizeMB;
 			FCC.StallWhenFileSystemThrottled = Configuration.ConstructorStallWhenFileSystemThrottled;
 			FCC.DisableResumeBelowMB = Configuration.ConstructorDisableResumeBelowMB;
-
+			FCC.bUseSHA1StageFilenames = !Configuration.bStageWithRawFilenames;
+			FCC.bRejectSymlinks = Configuration.bRejectSymlinks;
+			
 			TUniquePtr<FBuildPatchFileConstructor> FileConstructor(
 				new FBuildPatchFileConstructor(
 					FCC, 
@@ -1531,7 +1961,10 @@ namespace BuildPatchServices
 					InstallChunkSource->OnBeforeDeleteFile(FilePath);
 				});		
 			}
-
+			FDelegateHandle OnFileIsStagedHandle = FileConstructor->OnFileIsStaged().AddLambda([this](const FString& FilePath)
+				{
+					MessagePump->SendMessage(FInstallationFileAction{ FInstallationFileAction::EType::Staged, FilePath });
+				});
 			// Register controllables.
 			FScopedControllables ScopedControllables(&ThreadLock, Controllables, bIsPaused, bShouldAbort);
 			ScopedControllables.Register(FileConstructor.Get());
@@ -1551,10 +1984,10 @@ namespace BuildPatchServices
 			}
 
 			// Setup some weightings for the progress tracking
-			const bool bRepairOnly = ManifestSet->IsRepairOnly() && FilesToConstruct.Num() == 0;
+			const bool bExclusivelyRepairing = ManifestSet->IsExclusivelyRepairing() && FilesToConstruct.Num() == 0;
 			const float NumRequiredChunksFloat = ReferencedChunks.Num();
 			const bool bHasFileAttributes = ManifestSet->HasFileAttributes();
-			const float AttributesWeight = bHasFileAttributes ? bRepairOnly ? 1.0f / 50.0f : 1.0f / 20.0f : 0.0f;
+			const float AttributesWeight = bHasFileAttributes ? bExclusivelyRepairing ? 1.0f / 50.0f : 1.0f / 20.0f : 0.0f;
 			const float VerifyWeight = Configuration.VerifyMode == EVerifyMode::ShaVerifyAllFiles || Configuration.VerifyMode == EVerifyMode::ShaVerifyTouchedFiles ? 1.1f / 9.0f : 0.3f / 9.0f;
 			BuildProgress.SetStateWeight(EBuildPatchState::Downloading, 0.0f);
 			BuildProgress.SetStateWeight(EBuildPatchState::MovingToInstall, 0.05f);
@@ -1563,7 +1996,7 @@ namespace BuildPatchServices
 			BuildProgress.SetStateWeight(EBuildPatchState::BuildVerification, VerifyWeight);
 
 			// If this is a repair operation, start off with install and download complete
-			if (bRepairOnly)
+			if (bExclusivelyRepairing)
 			{
 				UE_LOG(LogBuildPatchServices, Log, TEXT("Performing a repair operation"));
 				BuildProgress.SetStateProgress(EBuildPatchState::Downloading, 1.0f);
@@ -1578,14 +2011,36 @@ namespace BuildPatchServices
 				InitializeTimer.Stop();
 			}
 
+			{
+				UE::TUniqueLock _(DownloadRateLimitLock);
+				DownloadRateUpdatedFn = [&DownloadConnectionCount, &CloudChunkSource](uint64 NewDownloadRateLimitBps)
+				{
+					CloudChunkSource->SetDownloadBytesPerSecondLimit(NewDownloadRateLimitBps);
+					DownloadConnectionCount->SetConnectionScaling(NewDownloadRateLimitBps == 0);
+				};
+			}
+
 			// Wait for the file constructor to complete
 			ConstructTimer.Start();
 			FileConstructor->Run();
 			ConstructTimer.Stop();
+			FileConstructor->OnFileIsStaged().Remove(OnFileIsStagedHandle);
+
+			{
+				UE::TUniqueLock _(DownloadRateLimitLock);
+				DownloadRateUpdatedFn.Reset();
+			}
+
 			if (OnBeforeDeleteFileHandle.IsValid())
 			{
 				FileConstructor->OnBeforeDeleteFile().Remove(OnBeforeDeleteFileHandle);
 			}
+			const bool bReportAnalytic = InstallerError->HasError() == false;
+			if (bReportAnalytic && bLastInstallIteration && InstallerAnalytics.IsValid() && FileConstructor->HasCorruptFiles())
+			{
+				InstallerAnalytics->RecordConstructionError({}, INDEX_NONE, TEXT("Serialised Verify Fail"));
+			}
+
 			UE_LOG(LogBuildPatchServices, Log, TEXT("File construction complete"));
 
 			if (Configuration.bInstallToMemory)
@@ -1615,25 +2070,7 @@ namespace BuildPatchServices
 		}
 
 		// Process some final stats.
-		{
-			FScopeLock Lock(&ThreadLock);
-			BuildStats.NumChunksDownloaded = DownloadServiceStatistics->GetNumSuccessfulChunkDownloads();
-			BuildStats.NumFailedDownloads = DownloadServiceStatistics->GetNumFailedChunkDownloads();
-			BuildStats.NumBadDownloads = CloudChunkSourceStatistics->GetNumCorruptChunkDownloads();
-			BuildStats.NumAbortedDownloads = CloudChunkSourceStatistics->GetNumAbortedChunkDownloads();
-			BuildStats.OverallRequestSuccessRate = CloudChunkSourceStatistics->GetDownloadSuccessRate();
-			BuildStats.NumChunksRecycled = InstallChunkSourceStatistics->GetNumSuccessfulChunkRecycles();
-			BuildStats.NumChunksReadFromChunkDbs = ChunkDbChunkSourceStatistics->GetNumSuccessfulLoads();
-			BuildStats.NumRecycleFailures = InstallChunkSourceStatistics->GetNumFailedChunkRecycles();
-			BuildStats.NumChunkDbChunksFailed = ChunkDbChunkSourceStatistics->GetNumFailedLoads();
-			TArray<float> HealthTimers = CloudChunkSourceStatistics->GetDownloadHealthTimers();
-			BuildStats.ExcellentDownloadHealthTime = HealthTimers[(int32)EBuildPatchDownloadHealth::Excellent];
-			BuildStats.GoodDownloadHealthTime = HealthTimers[(int32)EBuildPatchDownloadHealth::Good];
-			BuildStats.OkDownloadHealthTime = HealthTimers[(int32)EBuildPatchDownloadHealth::OK];
-			BuildStats.PoorDownloadHealthTime = HealthTimers[(int32)EBuildPatchDownloadHealth::Poor];
-			BuildStats.DisconnectedDownloadHealthTime = HealthTimers[(int32)EBuildPatchDownloadHealth::Disconnected];
-			BuildStats.ActiveRequestCountPeak = CloudChunkSourceStatistics->GetPeakRequestCount();
-		}
+		SetFinalInstallStats();
 
 		UE_LOG(LogBuildPatchServices, Log, TEXT("Staged install complete"));
 		const bool bReturnInstallationSuccess = !InstallerError->HasError();
@@ -1651,6 +2088,272 @@ namespace BuildPatchServices
 		return bReturnInstallationSuccess;
 	}
 
+	bool FBuildPatchInstaller::RunPreload()
+	{
+		// Reset our error and build progress
+		InstallerError.Reset(FInstallerErrorFactory::Create());
+		InstallerError->RegisterForErrors([this]() { CancelInstall(); });
+
+		UE_LOG(LogBuildPatchServices, Log, TEXT("Starting Preload"));
+
+		// Set weights for preload only
+		BuildProgress.SetStateWeight(EBuildPatchState::Resuming, 0.0f);
+		BuildProgress.SetStateWeight(EBuildPatchState::Downloading, 1.0f);
+		BuildProgress.SetStateWeight(EBuildPatchState::Installing, 1.0f);
+		BuildProgress.SetStateWeight(EBuildPatchState::MovingToInstall, 0.0f);
+		BuildProgress.SetStateWeight(EBuildPatchState::SettingAttributes, 0.0f);
+		BuildProgress.SetStateWeight(EBuildPatchState::BuildVerification, 0.0f);
+		BuildProgress.SetStateWeight(EBuildPatchState::RunningCustomUninstallAction, 0.0f);
+
+		// Mark all irrelevant steps complete
+		BuildProgress.SetStateProgress(EBuildPatchState::Resuming, 1.0f);
+		BuildProgress.SetStateProgress(EBuildPatchState::MovingToInstall, 1.0f);
+		BuildProgress.SetStateProgress(EBuildPatchState::SettingAttributes, 1.0f);
+		BuildProgress.SetStateProgress(EBuildPatchState::BuildVerification, 1.0f);
+		BuildProgress.SetStateProgress(EBuildPatchState::RunningCustomUninstallAction, 1.0f);
+
+		// We need an array of sorted, duplicate chunks, in order to produce the chunkDB files properly.
+		// So we use the TArray type function instead of TSet
+		TArray<FGuid> ReferencedChunks;
+		ManifestSet->GetReferencedChunks(ReferencedChunks);
+		Algo::Reverse(ReferencedChunks);
+		
+		// Default chunk store size to tie in with the default prefetch maxes for source configs.
+		int32 ChunkStoreMemorySize = 500;
+		// Load overridden size from config.
+		GConfig->GetInt(TEXT("Portal.BuildPatch"), TEXT("ChunkStoreMemorySize"), ChunkStoreMemorySize, GEngineIni);
+		// Clamp to sensible limits.
+		ChunkStoreMemorySize = FMath::Clamp<int32>(ChunkStoreMemorySize, 64, 2048);
+		// Cache the last download requirement in case we are running a retry.
+		PreviousTotalDownloadRequired += CloudChunkSourceStatistics->GetDownloadedBytes();
+		// Reset so that we don't double count data.
+		CloudChunkSourceStatistics->OnRequiredDataUpdated(0);
+		CloudChunkSourceStatistics->OnReceivedDataUpdated(0);
+		bAreResultsEncrypted = true;
+
+		// Scoped systems composition and execution.
+		{
+			TUniquePtr<IChunkDataSerialization> ChunkDataSerialization(FChunkDataSerializationFactory::Create(
+				FileSystem.Get(), Crypto.Get(), { EFeatureLevel::Latest, Configuration.EncryptionSecrets, FGuid() /* UE5 MERGE TODO : , CloudChunkSourceStatistics.Get()*/ }));
+
+			const TSet<FGuid> ReferencedChunksSet{ ReferencedChunks };
+
+			// We don't actually need a chunkdb source, but we'll use it to tell us what we already downloaded.
+			const TSet<FGuid> AlreadyFetchedChunks = [&]()
+				{
+					TUniquePtr<IConstructorChunkDbChunkSource> ChunkDbChunkSource(IConstructorChunkDbChunkSource::CreateChunkDbSource(
+						BuildChunkDbSourceConfig(Configuration, DataStagingDir),
+						FileSystem.Get(),
+						ReferencedChunks,
+						ChunkDataSerialization.Get(),
+						ChunkDbChunkSourceStatistics.Get()));
+					return ChunkDbChunkSource->GetAvailableChunks();
+				}();
+			const TSet<FGuid> InitialDownloadChunks = ReferencedChunksSet.Difference(AlreadyFetchedChunks);
+
+			// Set chunk counter stats.
+			{
+				FScopeLock Lock(&ThreadLock);
+				BuildStats.NumChunksRequired = ReferencedChunks.Num();
+				BuildStats.ChunksQueuedForDownload = InitialDownloadChunks.Num();
+				BuildStats.ChunksLocallyAvailable = AlreadyFetchedChunks.Num();
+				BuildStats.ChunksInChunkDbs = AlreadyFetchedChunks.Num();
+			}
+
+			// Initialisation complete.
+			BuildProgress.SetStateProgress(EBuildPatchState::Initializing, 1.0f);
+			FileConstructorStatistics->OnResumeCompleted();
+			InitializeTimer.Stop();
+
+			// If we don't have any InitialDownloadChunks, then we already finished building all of the chunkdb files necessary
+			TAtomic<bool> bChunkDbSuccess(true);
+			if (InitialDownloadChunks.Num() > 0)
+			{
+				TArray<FGuid> OrderedRequiredChunks;
+				TArray<FChunkDatabaseFile> ChunkDbFiles;
+				bChunkDbSuccess = false;
+				// Figure out the details for each individual chunkdb we want to save out.
+				const uint64 ChunkDbHeaderSize = ChunkDbHelpers::GetChunkDbHeaderSize();
+				const uint64 PerEntryHeaderSize = ChunkDbHelpers::GetPerChunkEntryHeaderSize();
+				const uint64 ChunkDbMaxSize = ConfigHelpers::LoadPreloadMaxChunkDbSize();
+				uint64 RequiredDiskSpace = 0;
+				uint64 AvailableFileSize = 0;
+				int32 ChunkDbFileIdx = 0;
+				OrderedRequiredChunks.Reserve(ReferencedChunks.Num());
+				for (const FGuid& DataId : ReferencedChunks)
+				{
+					if (!AlreadyFetchedChunks.Contains(DataId))
+					{
+						OrderedRequiredChunks.Add(DataId);
+						const uint64 DataSize = ManifestSet->GetDownloadSize(DataId) + PerEntryHeaderSize;
+						RequiredDiskSpace += DataSize;
+						if (AvailableFileSize < DataSize)
+						{
+							ChunkDbFileIdx = ChunkDbFiles.AddDefaulted();
+							AvailableFileSize = ChunkDbMaxSize - ChunkDbHeaderSize;
+							RequiredDiskSpace += ChunkDbHeaderSize;
+						}
+
+						ChunkDbFiles[ChunkDbFileIdx].DataList.Add(DataId);
+						AvailableFileSize -= DataSize;
+					}
+				}
+
+				// Count how many chunkdb files already exist, from a previous preload progress.
+				TArray<FString> PreExistingChunkDbs = ConfigHelpers::EnumeratePreloadChunkDbs(*DataStagingDir);
+
+				// Produce nicely named chunkDBs.
+				// Technically, there are mathematical solutions to this, however there can be floating point errors in log that cause edge cases there, so we'll just use the obvious simple method.
+				const int32 NumDigitsForParts = FString::Printf(TEXT("%d"), ChunkDbFiles.Num() + PreExistingChunkDbs.Num()).Len();
+				for (ChunkDbFileIdx = 0; ChunkDbFileIdx < ChunkDbFiles.Num(); ++ChunkDbFileIdx)
+				{
+					ChunkDbFiles[ChunkDbFileIdx].DatabaseFilename = DataStagingDir / FString::Printf(TEXT("clouddata.part%0*d.chunkdb"), NumDigitsForParts, ChunkDbFileIdx + 1 + PreExistingChunkDbs.Num());
+				}
+
+				// Setup file stats
+				{
+					FScopeLock Lock(&ThreadLock);
+					BuildStats.NumFilesInBuild = PreExistingChunkDbs.Num() + ChunkDbFiles.Num();
+					BuildStats.NumFilesOutdated = ChunkDbFiles.Num();
+				}
+
+				// Check we have enough disk space to proceed
+				uint64 AvailableDiskSpace = 0;
+				if (!FileConstructorHelpers::CheckRemainingDiskSpace(Configuration.InstallDirectory, RequiredDiskSpace, AvailableDiskSpace))
+				{
+					UE_LOG(LogBuildPatchServices, Error, TEXT("Preload: Out of HDD space. Needs %llu bytes, Free %llu bytes"), RequiredDiskSpace, AvailableDiskSpace);
+					InstallerError->SetError(
+						EBuildPatchInstallError::OutOfDiskSpace,
+						PreExistingChunkDbs.Num() > 0 ? DiskSpaceErrorCodes::DuringInstallation : DiskSpaceErrorCodes::InitialSpaceCheck,
+						0,
+						BuildPatchServices::GetDiskSpaceMessage(Configuration.InstallDirectory, RequiredDiskSpace, AvailableDiskSpace));
+					bChunkDbSuccess = false;
+				}
+				else
+				{
+					ConstructTimer.Start();
+					// Database writer dependencies.
+					TUniquePtr<IDownloadConnectionCount> DownloadConnectionCount(FDownloadConnectionCountFactory::Create(BuildConnectionCountConfig(), DownloadServiceStatistics.Get()));
+					TUniquePtr<IChunkReferenceTracker> ChunkReferenceTracker(FChunkReferenceTrackerFactory::Create(OrderedRequiredChunks));
+					TUniquePtr<IChunkEvictionPolicy> MemoryEvictionPolicy(FChunkEvictionPolicyFactory::Create(
+						ChunkReferenceTracker.Get()));
+					TUniquePtr<IMemoryChunkStoreStatistics> MemoryChunkStoreStatistics(FMemoryChunkStoreStatisticsFactory::Create(FileOperationTracker.Get()));
+					TUniquePtr<IMemoryChunkStore> CloudChunkStore(FMemoryChunkStoreFactory::Create(
+						ChunkStoreMemorySize,
+						MemoryEvictionPolicy.Get(),
+						nullptr,
+						MemoryChunkStoreStatistics.Get(),
+						ChunkReferenceTracker.Get()));
+						FCloudSourceConfig CloudSourceConfig = BuildCloudSourceConfig();
+					TUniquePtr<ICloudChunkSource> CloudChunkSource(FCloudChunkSourceFactory::Create(
+							MoveTemp(CloudSourceConfig),
+							Platform.Get(),
+							CloudChunkStore.Get(),
+							DownloadService.Get(),
+							ChunkReferenceTracker.Get(),
+							ChunkDataSerialization.Get(),
+							MessagePump.Get(),
+							InstallerError.Get(),
+							DownloadConnectionCount.Get(),
+							CloudChunkSourceStatistics.Get(),
+							ManifestSet.Get(),
+							InitialDownloadChunks));
+
+					// Start database writer that saves all the chunks to the chunkdbs.
+					FEvent* CompletionTrigger = FPlatformProcess::GetSynchEventFromPool(true);
+					FChunkDbWriterConfig ChunkDbWriterConfig(ChunkDbFiles);
+					ChunkDbWriterConfig.bUseTempFile = false;
+					ChunkDbWriterConfig.HeaderUpdateFrequency = 5.0f;
+					ChunkDbWriterConfig.bReserialise = false;
+					ChunkDbWriterConfig.bDeleteFilesOnFailure = false;
+					ChunkDbWriterConfig.bCallbackOnMainThread = false;
+					TUniquePtr<IChunkDatabaseWriter> ChunkDatabaseWriter(FChunkDatabaseWriterFactory::Create(
+						ChunkDbWriterConfig,
+						CloudChunkSource.Get(),
+						FileSystem.Get(),
+						InstallerError.Get(),
+						ChunkReferenceTracker.Get(),
+						ChunkDataSerialization.Get(),
+						FileConstructorStatistics.Get(),
+						[this](float Progress) { BuildProgress.SetStateProgress(EBuildPatchState::Installing, Progress); },
+						[&bChunkDbSuccess, CompletionTrigger](bool bInSuccess) { bChunkDbSuccess = bInSuccess; CompletionTrigger->Trigger(); }));
+
+					// Register controllables.
+					FScopedControllables ScopedControllables(&ThreadLock, Controllables, bIsPaused, bShouldAbort);
+					ScopedControllables.Register(CloudChunkSource.Get());
+					ScopedControllables.Register(ChunkDatabaseWriter.Get());
+
+					// Register download rate function.
+					{
+						UE::TUniqueLock _(DownloadRateLimitLock);
+						DownloadRateUpdatedFn = [&DownloadConnectionCount, &CloudChunkSource](uint64 NewDownloadRateLimitBps)
+						{
+							CloudChunkSource->SetDownloadBytesPerSecondLimit(NewDownloadRateLimitBps);
+							DownloadConnectionCount->SetConnectionScaling(NewDownloadRateLimitBps == 0);
+						};
+					}
+
+					// Wait for preload complete
+					CompletionTrigger->Wait();
+					FPlatformProcess::ReturnSynchEventToPool(CompletionTrigger);
+					ConstructTimer.Stop();
+				}
+
+				// Setup file stats
+				{
+					FScopeLock Lock(&ThreadLock);
+					BuildStats.NumFilesConstructed = FileConstructorStatistics->GetFilesConstructed();
+					BuildStats.ProcessRequiredDiskSpace = RequiredDiskSpace;
+					BuildStats.ProcessAvailableDiskSpace = AvailableDiskSpace;
+				}
+
+				// Process some final stats.
+				SetFinalInstallStats();
+			}
+
+			if (!bChunkDbSuccess)
+			{
+				// This call will be ignored, if the database writer correctly set an error code. Otherwise we will set unknown so we can track that an error ocurred and see we need more info.
+				InstallerError->SetError(EBuildPatchInstallError::FileConstructionFail, ConstructionErrorCodes::UnknownFail);
+			}
+			else
+			{
+				// Calculate if we want to try to reserve extra HDD space for the unlock process.
+				TArray<FString> AllChunkDbs = ConfigHelpers::EnumeratePreloadChunkDbs(*DataStagingDir);
+				uint64 ChunkDbSize = 0;
+				for (const FString& ChunkDbFile : AllChunkDbs)
+				{
+					int64 FileSize = 0;
+					FileSystem->GetFileSize(*ChunkDbFile, FileSize);
+					ChunkDbSize += FileSize;
+				}
+				TSet<FString> ExpectedFiles;
+				ManifestSet->GetExpectedFiles(ExpectedFiles);
+				const uint64 InstallSize = ManifestSet->GetTotalNewFileSize(ExpectedFiles);
+				if (InstallSize > ChunkDbSize)
+				{
+					const uint64 ReserveSize = InstallSize - ChunkDbSize;
+					FileSystem->SetFileSize(*HddSpaceReservationFile, ReserveSize);
+				}
+			}
+		}
+
+		UE_LOG(LogBuildPatchServices, Log, TEXT("Preload complete"));
+		const bool bReturnPreloadSuccess = !InstallerError->HasError();
+
+		// Ensure all progress complete
+		if (bReturnPreloadSuccess)
+		{
+			BuildProgress.SetStateProgress(EBuildPatchState::Initializing, 1.0f);
+			BuildProgress.SetStateProgress(EBuildPatchState::Resuming, 1.0f);
+			BuildProgress.SetStateProgress(EBuildPatchState::Downloading, 1.0f);
+			BuildProgress.SetStateProgress(EBuildPatchState::Installing, 1.0f);
+			InitializeTimer.Stop();
+		}
+
+		return bReturnPreloadSuccess;
+	}
+
 	bool FBuildPatchInstaller::RunPrerequisites()
 	{
 		TUniquePtr<IMachineConfig> MachineConfig(FMachineConfigFactory::Create(LocalMachineConfigFile, true));
@@ -1664,30 +2367,121 @@ namespace BuildPatchServices
 		return Prerequisites->RunPrereqs(ManifestSet.Get(), Configuration, InstallStagingDir, BuildProgress);
 	}
 
-	void FBuildPatchInstaller::CleanupEmptyDirectories(const FString& RootDirectory)
+	void FBuildPatchInstaller::SetFinalInstallStats()
+	{
+		FScopeLock Lock(&ThreadLock);
+		BuildStats.NumChunksDownloaded = DownloadServiceStatistics->GetNumSuccessfulChunkDownloads();
+		BuildStats.NumFailedDownloads = DownloadServiceStatistics->GetNumFailedChunkDownloads();
+		BuildStats.NumBadDownloads = CloudChunkSourceStatistics->GetNumCorruptChunkDownloads();
+		BuildStats.NumAbortedDownloads = CloudChunkSourceStatistics->GetNumAbortedChunkDownloads();
+		BuildStats.OverallRequestSuccessRate = CloudChunkSourceStatistics->GetDownloadSuccessRate();
+		BuildStats.NumChunksRecycled = InstallChunkSourceStatistics->GetNumSuccessfulChunkRecycles();
+		BuildStats.NumChunksReadFromChunkDbs = ChunkDbChunkSourceStatistics->GetNumSuccessfulLoads();
+		BuildStats.NumRecycleFailures = InstallChunkSourceStatistics->GetNumFailedChunkRecycles();
+		BuildStats.NumChunkDbChunksFailed = ChunkDbChunkSourceStatistics->GetNumFailedLoads();
+		TArray<float> HealthTimers = CloudChunkSourceStatistics->GetDownloadHealthTimers();
+		BuildStats.ExcellentDownloadHealthTime = HealthTimers[(int32)EBuildPatchDownloadHealth::Excellent];
+		BuildStats.GoodDownloadHealthTime = HealthTimers[(int32)EBuildPatchDownloadHealth::Good];
+		BuildStats.OkDownloadHealthTime = HealthTimers[(int32)EBuildPatchDownloadHealth::OK];
+		BuildStats.PoorDownloadHealthTime = HealthTimers[(int32)EBuildPatchDownloadHealth::Poor];
+		BuildStats.DisconnectedDownloadHealthTime = HealthTimers[(int32)EBuildPatchDownloadHealth::Disconnected];
+		BuildStats.ActiveRequestCountPeak = CloudChunkSourceStatistics->GetPeakRequestCount();
+		BuildStats.ActiveRequestCountAvg = CloudChunkSourceStatistics->GetAvgRequestCount();
+		const FString UserId = FPlatformMisc::GetEpicAccountId();
+		const FString Variant = Experiments::GetVariant(Experiments::DownloadSpeed(), UserId);
+		BuildStats.ExperimentName = FString(Experiments::DownloadSpeed().ExperimentName);
+		BuildStats.ExerimentVariantName = *Variant;
+		BuildStats.TotalDecryptTime = CloudChunkSourceStatistics->GetTotalDecryptTime();
+		BuildStats.CloudDirectoryStats = CloudChunkSourceStatistics->GetCloudDirectoryStats();
+	}
+
+	class FParallelCleanUpDirectoryEnumerator : public IPlatformFile::FDirectoryVisitor
 	{
-		TArray<FString> SubDirNames;
-		IFileManager::Get().FindFiles(SubDirNames, *(RootDirectory / TEXT("*")), false, true);
-		for(auto DirName : SubDirNames)
+	public:
+		FParallelCleanUpDirectoryEnumerator(const FString& InRootDirectory) : RootDirectory(InRootDirectory) {}
+
+		virtual bool Visit(const TCHAR* FilenameOrDirectory, bool bIsDirectory) override
+		{
+			if (bIsDirectory)
+			{
+				FString Directory = FilenameOrDirectory;
+				DirectoryFutures.Enqueue(Async(EAsyncExecution::ThreadPool, [this, Directory = MoveTemp(Directory)]()
+					{
+						struct FEmptyDirVisitor : public IPlatformFile::FDirectoryVisitor
+						{
+							// Once visited means either a folder or a file exists inside of the dir
+							virtual bool Visit(const TCHAR*, bool) override { return false; }
+						};
+						static FEmptyDirVisitor EmptyDirVisitor;
+						if (IFileManager::Get().IterateDirectory(*Directory, EmptyDirVisitor))
+						{
+							return Directory;
+						}
+						IFileManager::Get().IterateDirectory(*Directory, *this);
+						return FString();
+					}));
+			}
+			return true;
+		}
+
+		void GetCleanUpFolders(TSet<FString>& Results)
 		{
-			CleanupEmptyDirectories(*(RootDirectory / DirName));
+			TFuture<FString> Future;
+			while (DirectoryFutures.Dequeue(Future))
+			{
+				FString Dir = Future.Get();
+				if (!Dir.IsEmpty())
+				{
+					Results.Add(MoveTemp(Dir));
+				}
+			}
 		}
 
-		TArray<FString> SubFileNames;
-		IFileManager::Get().FindFilesRecursive(SubFileNames, *RootDirectory, TEXT("*.*"), true, false);
-		if (SubFileNames.Num() == 0)
+	private:
+
+		static bool IsSlashOrBackslash(TCHAR C)
 		{
-	#if PLATFORM_MAC
+			return C == TEXT('/') || C == TEXT('\\');
+		}
+
+	private:
+		const FString& RootDirectory;
+		TQueue<TFuture<FString>, EQueueMode::Mpsc> DirectoryFutures;
+	};
+
+	void CleanupEmptyDirectoriesImpl(const FString& RootDirectory, TSet<FString> HandledFolders = {})
+	{
+		TSet<FString> CleanUpFolders;
+		FParallelCleanUpDirectoryEnumerator DirectoryEnumerator(RootDirectory);
+		IFileManager::Get().IterateDirectory(*RootDirectory, DirectoryEnumerator);
+		DirectoryEnumerator.GetCleanUpFolders(CleanUpFolders);
+		// to avoid recursion if we cannot delete a folder
+		CleanUpFolders = CleanUpFolders.Difference(HandledFolders);
+
+		for (const FString& FolderToDelete : CleanUpFolders)
+		{
+#if PLATFORM_MAC
 			// On Mac we need to delete the .DS_Store file, but FindFiles() skips .DS_Store files.
-			IFileManager::Get().Delete(*(RootDirectory / TEXT(".DS_Store")), false, true);
-	#endif
+			IFileManager::Get().Delete(*(FolderToDelete / TEXT(".DS_Store")), false, true);
+#endif
 
-			bool bDeleteSuccess = IFileManager::Get().DeleteDirectory(*RootDirectory, false, true);
+			bool bDeleteSuccess = IFileManager::Get().DeleteDirectory(*FolderToDelete, false, true);
 			const uint32 LastError = FPlatformMisc::GetLastError();
-			UE_LOG(LogBuildPatchServices, Log, TEXT("Deleted Empty Folder (%u,%u) %s"), bDeleteSuccess ? 1 : 0, LastError, *RootDirectory);
+			UE_LOG(LogBuildPatchServices, Log, TEXT("Deleted Empty Folder (%u,%u) %s"), bDeleteSuccess ? 1 : 0, LastError, *FolderToDelete);
+			HandledFolders.Add(FolderToDelete);
+		}
+
+		if (CleanUpFolders.Num() > 0)
+		{
+			CleanupEmptyDirectoriesImpl(RootDirectory, MoveTemp(HandledFolders));
 		}
 	}
 
+	void FBuildPatchInstaller::CleanupEmptyDirectories(const FString& RootDirectory)
+	{
+		CleanupEmptyDirectoriesImpl(RootDirectory);
+	}
+
 	void FBuildPatchInstaller::FilterToExistingFiles(const FString& RootDirectory, const TMap<FString, FString>& PerFileSubdirectories, TSet<FString>& Files, bool bResolveInstallationFileName)
 	{
 		for (auto FilesIt = Files.CreateIterator(); FilesIt; ++FilesIt)
@@ -1732,6 +2526,16 @@ namespace BuildPatchServices
 		CopyErrorCode = 0;
 		while (RelocateRetries >= 0 && !bRelocateSuccess)
 		{
+			const bool bIsDirectory = FileSystem->DirectoryExists(*ToFullFilename);
+			if (bIsDirectory)
+			{
+				const bool bDeleteDirSuccess = FileSystem->DeleteDirectoryRecursively(*ToFullFilename);
+				if (!bDeleteDirSuccess)
+				{
+					const uint32 ErrorCode = FPlatformMisc::GetLastError();
+					UE_LOG(LogBuildPatchServices, Warning, TEXT("Failed to delete directory %s (%d), trying move anyway.."), *ToFullFilename, ErrorCode);
+				}
+			}
 			bRelocateSuccess = FileSystem->MoveFile(*ToFullFilename, *FromFullFilename);
 			RenameErrorCode = FPlatformMisc::GetLastError();
 			if (!bRelocateSuccess)
@@ -1743,7 +2547,7 @@ namespace BuildPatchServices
 				{
 					FileSystem->DeleteFile(*FromFullFilename);
 				}
-				else if((--RelocateRetries) >= 0)
+				else if ((--RelocateRetries) >= 0)
 				{
 					UE_LOG(LogBuildPatchServices, Warning, TEXT("Failed to copy too (%d), retying after 0.5 sec.."), CopyErrorCode);
 					FPlatformProcess::Sleep(0.5f);
@@ -1757,14 +2561,61 @@ namespace BuildPatchServices
 		return bRelocateSuccess;
 	}
 
+	bool FBuildPatchInstaller::CopyFileWithRetries(const FString& ToFullFilename, const FString& FromFullFilename, uint32& CopyErrorCode)
+	{
+		int32 RelocateRetries = ConfigHelpers::NumFileMoveRetries();
+		bool bRelocateSuccess = false;
+		CopyErrorCode = 0;
+		while (RelocateRetries >= 0 && !bRelocateSuccess)
+		{
+			bRelocateSuccess = FileSystem->CopyFile(*ToFullFilename, *FromFullFilename);
+			CopyErrorCode = FPlatformMisc::GetLastError();
+			if (!bRelocateSuccess && (--RelocateRetries) >= 0)
+			{
+				UE_LOG(LogBuildPatchServices, Warning, TEXT("Failed to copy file %s (%d), retying after 0.5 sec.."), *FromFullFilename, CopyErrorCode);
+				FPlatformProcess::Sleep(0.5f);
+			}
+		}
+		if (!bRelocateSuccess)
+		{
+			UE_LOG(LogBuildPatchServices, Error, TEXT("Failed to copy file %s (%d), installer will exit."), *FromFullFilename, CopyErrorCode);
+		}
+		return bRelocateSuccess;
+	}
+
+	bool FBuildPatchInstaller::RunCustomUninstallAction()
+	{
+		BuildProgress.SetStateProgress(EBuildPatchState::RunningCustomUninstallAction, 0.0f);
+
+		const bool bSkipThisStep = Configuration.InstallMode == EInstallMode::PrereqOnly
+			|| InstallerError->HasError()
+			|| InstallerError->IsCancelled();
+
+		if (bSkipThisStep)
+		{
+			UE_LOG(LogBuildPatchServices, Log, TEXT("Skipping custom uninstall action"));
+			BuildProgress.SetStateProgress(EBuildPatchState::RunningCustomUninstallAction, 1.0f);
+			return true;
+		}
+
+		UninstallActionTimer.Start();
+		TUniquePtr<ICustomUninstallAction> CustomUninstallAction(FCustomUninstallActionFactory::Create(InstallerError.Get(), FileSystem.Get(), Platform.Get()));
+		bool bRunSuccessfull = CustomUninstallAction->RunAction(ManifestSet.Get(), Configuration, InstallStagingDir, BuildProgress);
+		UninstallActionTimer.Stop();
+
+		return bRunSuccessfull;
+	}
+
 	bool FBuildPatchInstaller::RunBackupAndMove()
 	{
 		BuildProgress.SetStateProgress(EBuildPatchState::MovingToInstall, 0.0f);
 		// We skip this step if performing stage only
 		bool bRunBackupAndMoveSuccess = true;
-		if (Configuration.InstallMode == EInstallMode::StageFiles || 
-			(bFirstInstallIteration && Configuration.InstallMode == EInstallMode::PrereqOnly) ||
-			Configuration.bInstallToMemory)
+		if (Configuration.InstallMode == EInstallMode::StageFiles
+			|| (bFirstInstallIteration && Configuration.InstallMode == EInstallMode::PrereqOnly)
+			|| ManifestSet->IsExclusivelyVerifying()
+			|| bAreResultsEncrypted
+			|| Configuration.bInstallToMemory)
 		{
 			UE_LOG(LogBuildPatchServices, Log, TEXT("Skipping backup and stage relocation"));
 			BuildProgress.SetStateProgress(EBuildPatchState::MovingToInstall, 1.0f);
@@ -1780,7 +2631,44 @@ namespace BuildPatchServices
 				// Get list of all expected files
 				TSet<FString> FilesToRelocate;
 				ManifestSet->GetExpectedFiles(FilesToRelocate);
-				FilterToExistingFiles(InstallStagingDir, NewPerFileSubdirectories, FilesToRelocate, false);
+
+				// Filter symlinks from files to relocate
+				TSet<FString> SymlinksToCreate;
+				for (auto FilesIt = FilesToRelocate.CreateIterator(); FilesIt; ++FilesIt)
+				{
+					const FFileManifest* FileManifest = ManifestSet->GetNewFileManifest(**FilesIt);
+					if (FileManifest != nullptr && !FileManifest->SymlinkTarget.IsEmpty())
+					{
+						SymlinksToCreate.Add(*FilesIt);
+						FilesIt.RemoveCurrent();
+					}
+				}
+
+				TMap<FSHAHash, uint32> FileHashToCountMap;
+				if (!Configuration.bStageWithRawFilenames)
+				{
+					// Filter to existing files with SHA1 hash value as filename
+					for (auto FilesIt = FilesToRelocate.CreateIterator(); FilesIt; ++FilesIt)
+					{
+						if (const FFileManifest* FileManifest = ManifestSet->GetNewFileManifest(**FilesIt))
+						{
+							FString FileToRelocateHashStringValue = GetSHA1Filename(*FileManifest);
+
+							if (!FileSystem->FileExists(*(InstallStagingDir / *FileToRelocateHashStringValue)))
+							{
+								FilesIt.RemoveCurrent();
+								continue;
+							}
+
+							// We need to map each hash to its file count for proper relocation it
+							FileHashToCountMap.FindOrAdd(FileManifest->SHA1Hash)++;
+						}
+					}
+				}
+				else
+				{
+					FilterToExistingFiles(InstallStagingDir, NewPerFileSubdirectories, FilesToRelocate, false);
+				}
 				FilesToRelocate.Sort(TLess<FString>());
 
 				// First handle files that should be removed for patching
@@ -1791,7 +2679,7 @@ namespace BuildPatchServices
 				FilesToRemove.Sort(TLess<FString>());
 
 				// Counters for progress tracking.
-				float NumOperationsFloat = FilesToRelocate.Num() + FilesToRemove.Num();
+				float NumOperationsFloat = FilesToRelocate.Num() + FilesToRemove.Num() + SymlinksToCreate.Num();
 				float PerformedOperationsFloat = 0;
 
 				// Add to build stats
@@ -1829,9 +2717,13 @@ namespace BuildPatchServices
 				// Perform all of the relocations
 				for (auto FileToRelocate = FilesToRelocate.CreateConstIterator(); FileToRelocate && !InstallerError->HasError(); ++FileToRelocate)
 				{
-					const FString SrcFilename = InstallStagingDir / *FileToRelocate;
 					const FString DestFilename = FBuildPatchUtils::ResolveInstallationFileName(Configuration.InstallDirectory, NewPerFileSubdirectories, *FileToRelocate);
+					const bool bOldDirectoryExists = FileSystem->DirectoryExists(*DestFilename);
 					const bool bOldFileExists = FileSystem->FileExists(*DestFilename);
+					if (bOldDirectoryExists)
+					{
+						CleanupEmptyDirectories(DestFilename);
+					}
 					if (bOldFileExists)
 					{
 						BackupFileIfNecessary(*FileToRelocate, NewPerFileSubdirectories);
@@ -1846,7 +2738,42 @@ namespace BuildPatchServices
 					}
 					uint32 RenameErrorCode = 0;
 					uint32 CopyErrorCode = 0;
-					bool bRelocateSuccess = RelocateFileWithRetries(DestFilename, SrcFilename, RenameErrorCode, CopyErrorCode);
+
+					bool bRelocateSuccess = false;
+
+					if (!Configuration.bStageWithRawFilenames)
+					{
+						if (const FFileManifest* FileManifest = ManifestSet->GetNewFileManifest(*FileToRelocate))
+						{
+							FString FileHashStringValue = GetSHA1Filename(*FileManifest);
+
+							const FString SrcFilename = InstallStagingDir / *FileHashStringValue;
+
+							// several files can have same SHA1 hash
+							// so we need to copy this files Count times to dest dir
+							// and at the end move it to dest dir
+							uint32* Count = FileHashToCountMap.Find(FileManifest->SHA1Hash);
+							if (*Count == 1)
+							{
+								bRelocateSuccess = RelocateFileWithRetries(DestFilename, SrcFilename, RenameErrorCode, CopyErrorCode);
+							}
+							else
+							{
+								bRelocateSuccess = CopyFileWithRetries(DestFilename, SrcFilename, CopyErrorCode);
+								*Count = *Count - 1;
+							}
+						}
+						else
+						{
+							UE_LOG(LogBuildPatchServices, Error, TEXT("Unable to find a manifest for a file %s"), **FileToRelocate);
+						}
+					}
+					else
+					{
+						const FString SrcFilename = InstallStagingDir / *FileToRelocate;
+						bRelocateSuccess = RelocateFileWithRetries(DestFilename, SrcFilename, RenameErrorCode, CopyErrorCode);
+					}
+
 					if (bRelocateSuccess)
 					{
 						FilesInstalled.Add(*FileToRelocate);
@@ -1865,6 +2792,53 @@ namespace BuildPatchServices
 					PerformedOperationsFloat += 1.0f;
 				}
 
+				// Create all symlinks
+				for (auto SymlinkToCreate = SymlinksToCreate.CreateConstIterator(); SymlinkToCreate && !InstallerError->HasError(); ++SymlinkToCreate)
+				{
+					const FString SymlinkFilename = Configuration.InstallDirectory / *SymlinkToCreate;
+
+					const bool bOldFileExists = FileSystem->FileExists(*SymlinkFilename);
+					if (bOldFileExists)
+					{
+						uint32 ErrorCode = 0;
+						bool bDeleteSuccess = RemoveFileWithRetries(SymlinkFilename, ErrorCode);
+						if (!bDeleteSuccess)
+						{
+							InstallerError->SetError(EBuildPatchInstallError::MoveFileToInstall, MoveErrorCodes::DeletePrevFileFailed, ErrorCode);
+							bRunBackupAndMoveSuccess = false;
+							break;
+						}
+					}
+
+					bool bCreateSuccess = false;
+					uint32 ErrorCode = 0;
+
+					const FFileManifest* FileManifest = ManifestSet->GetNewFileManifest(*SymlinkToCreate);
+					FileSystem->MakeDirectory(*FPaths::GetPath(SymlinkFilename));
+
+#if PLATFORM_MAC || PLATFORM_LINUX
+					unlink(TCHAR_TO_UTF8(*SymlinkFilename));
+					bCreateSuccess = symlink(TCHAR_TO_UTF8(*FileManifest->SymlinkTarget), TCHAR_TO_UTF8(*SymlinkFilename)) == 0;
+					ErrorCode = errno;
+#endif
+
+					if (!bCreateSuccess)
+					{
+						// Report if we could not create the symlink.
+						InstallerError->SetError(EBuildPatchInstallError::FileConstructionFail, ConstructionErrorCodes::SymlinkCreateFail, ErrorCode);
+						bRunBackupAndMoveSuccess = false;
+						break;
+					}
+
+					FilesInstalled.Add(*SymlinkToCreate);
+					FileOperationTracker->OnFileStateUpdate(*SymlinkToCreate, EFileOperationState::Installed);
+					FInstallationFileAction::EType Action = OldFilesRemovedBySystem.Contains(*SymlinkToCreate) || bOldFileExists ? FInstallationFileAction::EType::Updated : FInstallationFileAction::EType::Added;
+					MessagePump->SendMessage(FInstallationFileAction{ Action, *SymlinkToCreate });
+
+					BuildProgress.SetStateProgress(EBuildPatchState::MovingToInstall, PerformedOperationsFloat / NumOperationsFloat);
+					PerformedOperationsFloat += 1.0f;
+				}
+
 				bRunBackupAndMoveSuccess = bRunBackupAndMoveSuccess && (InstallerError->HasError() == false);
 				if (bRunBackupAndMoveSuccess)
 				{
@@ -1879,14 +2853,14 @@ namespace BuildPatchServices
 
 	bool FBuildPatchInstaller::RunFileAttributes()
 	{
-		if (Configuration.bInstallToMemory)
+		if (ManifestSet->IsExclusivelyVerifying() || bAreResultsEncrypted || Configuration.bInstallToMemory)
 		{
 			return true;
 		}
 
 		// Only provide stage directory if stage-only mode
 		FString EmptyString;
-		FString& OptionalStageDirectory = Configuration.InstallMode == EInstallMode::StageFiles ? InstallStagingDir : EmptyString;
+		const FString& OptionalStageDirectory = Configuration.InstallMode == EInstallMode::StageFiles ? InstallStagingDir : EmptyString;
 
 		// Construct the attributes class
 		FileAttributesTimer.Start();
@@ -1909,6 +2883,14 @@ namespace BuildPatchServices
 			return true;
 		}
 
+		// If we are storing encrypted data, we can't verify. This will be done by a later 'unlocking' stage.
+		if (bAreResultsEncrypted)
+		{
+			UE_LOG(LogBuildPatchServices, Log, TEXT("Verify stage skipping due to preloading encrypted data."));
+			BuildProgress.SetStateProgress(EBuildPatchState::BuildVerification, 1.0f);
+			return true;
+		}
+
 		// Make sure this function can never be parallelized
 		static FCriticalSection SingletonFunctionLockCS;
 		const bool bShouldLock = !Configuration.bAllowConcurrentExecution;
@@ -1926,11 +2908,19 @@ namespace BuildPatchServices
 
 		// Verify the build
 		EVerifyResult VerifyResult = Verifier->Verify(CorruptFiles);
-		const bool bVerifySuccessful = VerifyResult == EVerifyResult::Success;
+		bool bVerifySuccessful = VerifyResult == EVerifyResult::Success;
 		if (!bVerifySuccessful)
 		{
-			UE_LOG(LogBuildPatchServices, Error, TEXT("Build verification failed on %u file(s)"), CorruptFiles.Num());
-			InstallerError->SetError(EBuildPatchInstallError::BuildVerifyFail, InstallerHelpers::GetVerifyErrorCode(VerifyResult));
+			if (VerifyResult == EVerifyResult::Aborted)
+			{
+				UE_LOG(LogBuildPatchServices, Error, TEXT("User aborted verification"));
+				InstallerError->SetError(EBuildPatchInstallError::UserCanceled, UserCancelErrorCodes::UserRequested);
+			}
+			else
+			{
+				UE_LOG(LogBuildPatchServices, Error, TEXT("Build verification failed on %u file(s)"), CorruptFiles.Num());
+				InstallerError->SetError(EBuildPatchInstallError::BuildVerifyFail, InstallerHelpers::GetVerifyErrorCode(VerifyResult));
+			}
 		}
 		TMap<EVerifyError, int32> VerifyErrorCounts = VerifierStatistics->GetVerifyErrorCounts();
 		for (const TPair<EVerifyError, int32>& VerifyErrorCount : VerifyErrorCounts)
@@ -1950,17 +2940,41 @@ namespace BuildPatchServices
 		{
 			for (const FString& CorruptFile : CorruptFiles)
 			{
-				BackupFileIfNecessary(CorruptFile, NewPerFileSubdirectories, true);
-				FString StagedFile = InstallStagingDir / CorruptFile;
-				if (FileSystem->FileExists(*StagedFile))
+				// Getting staged file name, it can be different for cases when SHA1 are used instead
+				FString StagedFilenameToUse;
+				if (!Configuration.bStageWithRawFilenames)
+				{
+					const FFileManifest* FileManifest = ManifestSet->GetNewFileManifest(CorruptFile);
+					if (FileManifest == nullptr)
+					{
+						UE_LOG(LogBuildPatchServices, Error, TEXT("Can't get manifest for file %s"), *CorruptFile);
+						continue;
+					}
+
+					StagedFilenameToUse = GetSHA1Filename(*FileManifest);
+				}
+				else
+				{
+					StagedFilenameToUse = CorruptFile;
+				}
+
+				// Deleting files in a staged directory
+				FString StagingPathToUse = InstallStagingDir / StagedFilenameToUse;
+				if (!FileSystem->DeleteFile(*StagingPathToUse))
 				{
-					FileSystem->DeleteFile(*StagedFile);
+					uint32 ErrorCode = FPlatformMisc::GetLastError();
+					UE_LOG(LogBuildPatchServices, Error, TEXT("Failed to delete corrupted file %s with error code %d"), *StagingPathToUse, ErrorCode);
+					InstallerError->SetError(EBuildPatchInstallError::BuildVerifyFail, VerifyErrorCodes::DeleteCorruptedFileFailed, ErrorCode);
+					bVerifySuccessful = false;
 				}
+
+				// Deleting files in the destination install directory when non-stagefiles mode is used
 				if (Configuration.InstallMode != EInstallMode::StageFiles)
 				{
 					FString InstalledFile = FBuildPatchUtils::ResolveInstallationFileName(Configuration.InstallDirectory, NewPerFileSubdirectories, CorruptFile);
 					if (FileSystem->FileExists(*InstalledFile))
 					{
+						BackupFileIfNecessary(CorruptFile, NewPerFileSubdirectories, true);
 						FileSystem->DeleteFile(*InstalledFile);
 						OldFilesRemovedBySystem.Add(CorruptFile);
 					}
@@ -2022,8 +3036,8 @@ namespace BuildPatchServices
 			const int64 OriginalFileSize = OldFileManifest ? OldFileManifest->FileSize : INDEX_NONE;
 			const int64 NewFileSize = NewFileManifest ? NewFileManifest->FileSize : INDEX_NONE;
 			const FSHAHash HashZero;
-			const FSHAHash& HashOld = OldFileManifest ? OldFileManifest->FileHash : HashZero;
-			const FSHAHash& HashNew = NewFileManifest ? NewFileManifest->FileHash : HashZero;
+			const FSHAHash& HashOld = OldFileManifest ? OldFileManifest->SHA1Hash : HashZero;
+			const FSHAHash& HashNew = NewFileManifest ? NewFileManifest->SHA1Hash : HashZero;
 			const bool bFileSizeDiffers = OriginalFileSize != InstalledFilesize && NewFileSize != InstalledFilesize;
 			bUserEditedFile = bFileSizeDiffers || FBuildPatchUtils::VerifyFile(FileSystem.Get(), InstalledFilename, HashOld, HashNew) == 0;
 		}
@@ -2046,12 +3060,13 @@ namespace BuildPatchServices
 
 	double FBuildPatchInstaller::GetDownloadSpeed() const
 	{
-		return DownloadSpeedRecorder->GetAverageSpeed(ConfigHelpers::DownloadSpeedAverageTime());
+		const bool bThrottlingEnabled = DownloadRateLimitBps != 0;
+		return DownloadSpeedRecorder->GetAverageSpeed(ConfigHelpers::DownloadSpeedAverageTime(bThrottlingEnabled), bThrottlingEnabled);
 	}
 
 	int64 FBuildPatchInstaller::GetTotalDownloadRequired() const
 	{
-		return CloudChunkSourceStatistics->GetRequiredDownloadSize() + PreviousTotalDownloadRequired.GetValue();
+		return CloudChunkSourceStatistics->GetRequiredDownloadSize() + PreviousTotalDownloadRequired;
 	}
 
 	int64 FBuildPatchInstaller::GetTotalDownloaded() const
@@ -2066,8 +3081,7 @@ namespace BuildPatchServices
 
 	bool FBuildPatchInstaller::IsCanceled() const
 	{
-		FScopeLock Lock(&ThreadLock);
-		return BuildStats.FailureType == EBuildPatchInstallError::UserCanceled;
+		return InstallerError->IsCancelled();
 	}
 
 	bool FBuildPatchInstaller::IsPaused() const
@@ -2098,6 +3112,42 @@ namespace BuildPatchServices
 		return false;
 	}
 
+	bool FBuildPatchInstaller::IsRepair() const
+	{
+		for (const FBuildPatchInstallerAction& InstallerAction : InstallerActions)
+		{
+			if (InstallerAction.IsRepair())
+			{
+				return true;
+			}
+		}
+		return false;
+	}
+
+	bool FBuildPatchInstaller::IsVerifyOnly() const
+	{
+		for (const FBuildPatchInstallerAction& InstallerAction : InstallerActions)
+		{
+			if (InstallerAction.IsVerifyOnly())
+			{
+				return true;
+			}
+		}
+		return false;
+	}
+
+	bool FBuildPatchInstaller::IsUninstall() const
+	{
+		for (const FBuildPatchInstallerAction& InstallerAction : InstallerActions)
+		{
+			if (!InstallerAction.IsUninstall())
+			{
+				return false;
+			}
+		}
+		return true;
+	}
+
 	bool FBuildPatchInstaller::CompletedSuccessfully() const
 	{
 		return IsComplete() && bSuccess;
@@ -2189,13 +3239,13 @@ namespace BuildPatchServices
 	}
 
 	PRAGMA_DISABLE_DEPRECATION_WARNINGS
-	FText FBuildPatchInstaller::GetStatusText() const
+		FText FBuildPatchInstaller::GetStatusText() const
 	{
 		return BuildPatchServices::StateToText(GetState());
 	}
 	PRAGMA_ENABLE_DEPRECATION_WARNINGS
 
-	float FBuildPatchInstaller::GetUpdateProgress() const
+		float FBuildPatchInstaller::GetUpdateProgress() const
 	{
 		return BuildProgress.GetProgress();
 	}
@@ -2203,7 +3253,17 @@ namespace BuildPatchServices
 	FBuildInstallStats FBuildPatchInstaller::GetBuildStatistics() const
 	{
 		FScopeLock Lock(&ThreadLock);
-		return BuildStats;
+		FBuildInstallStats CurrentStats(BuildStats);
+
+		if (!IsComplete())
+		{
+			CurrentStats.AverageDownloadSpeed = DownloadSpeedRecorder->GetAverageSpeed(TNumericLimits<float>::Max());
+			CurrentStats.PeakDownloadSpeed = DownloadSpeedRecorder->GetPeakSpeed();
+			CurrentStats.NumChunksDownloaded = DownloadServiceStatistics->GetNumSuccessfulChunkDownloads();
+			CurrentStats.NumFailedDownloads = DownloadServiceStatistics->GetNumFailedChunkDownloads();
+		}
+
+		return CurrentStats;
 	}
 
 	EBuildPatchDownloadHealth FBuildPatchInstaller::GetDownloadHealth() const
@@ -2259,6 +3319,7 @@ namespace BuildPatchServices
 		}
 		// Set pause state on pausable process timers.
 		ConstructTimer.SetPause(bIsPaused);
+		UninstallActionTimer.SetPause(bIsPaused);
 		MoveFromStageTimer.SetPause(bIsPaused);
 		FileAttributesTimer.SetPause(bIsPaused);
 		VerifyTimer.SetPause(bIsPaused);
