p4 integrate //UE5/Dev-Distro-5.5/Engine/Source/Programs/BuildPatchTool/... //UE5/Main/Engine/Source/Programs/BuildPatchTool/...

diff --git a/Engine/Source/Runtime/Online/BuildPatchServices/Private/BuildPatchInstaller.h b/Engine/Source/Runtime/Online/BuildPatchServices/Private/BuildPatchInstaller.h
--- a/Engine/Source/Runtime/Online/BuildPatchServices/Private/BuildPatchInstaller.h
+++ b/Engine/Source/Runtime/Online/BuildPatchServices/Private/BuildPatchInstaller.h
@@ -23,6 +23,9 @@
 #include "BuildPatchManifest.h"
 #include "BuildPatchProgress.h"
 #include "IBuildManifestSet.h"
+#include "Installer/InstallChunkSource.h"
+
+#include <atomic>
 
 PRAGMA_DISABLE_DEPRECATION_WARNINGS
 
@@ -61,11 +64,11 @@ namespace BuildPatchServices
 		, public TSharedFromThis<FBuildPatchInstaller, ESPMode::ThreadSafe>
 	{
 	private:
-		// Simple unique installer guid to help logging functions have context.
-		FGuid SessionId;
+		// The installer configuration.
+		const FBuildInstallerConfiguration Configuration;
 
-		// Hold a pointer to my thread for easier deleting.
-		IBuildInstallerThread* Thread;
+		// Simple unique installer guid to help logging functions have context.
+		const FGuid SessionId;
 
 		// The delegates that we will be calling when started.
 		const FBuildPatchInstallerDelegate StartDelegate;
@@ -73,41 +76,54 @@ namespace BuildPatchServices
 		// The delegates that we will be calling on complete.
 		const FBuildPatchInstallerDelegate CompleteDelegate;
 
-		// The installer configuration.
-		FBuildInstallerConfiguration Configuration;
-
 		// The Configuration.InstallerActions array converted into private class type.
+		// Cannot be const due to Delta Optimisation modifying the manifests.
 		TArray<FBuildPatchInstallerAction> InstallerActions;
 
 		// The directory created in staging, to store local patch data.
-		FString DataStagingDir;
+		const FString DataStagingDir;
 
 		// The directory created in staging to construct install files to.
-		FString InstallStagingDir;
+		const FString InstallStagingDir;
 
 		// The directory created in staging to store any required meta, such as resume info.
-		FString MetaStagingDir;
+		const FString MetaStagingDir;
 
 		// The filename used to mark a previous install that did not complete but moved staged files into the install directory.
-		FString PreviousMoveMarker;
+		const FString PreviousMoveMarker;
+
+		// The filename used to reserve diskspace when preloading.
+		const FString HddSpaceReservationFile;
+
+		// Map of registered installations.
+		const TMultiMap<FString, FBuildPatchAppManifestRef> InstallationInfo;
 
 		// A critical section to protect variables.
 		mutable FCriticalSection ThreadLock;
 
+		// Hold a pointer to my thread for easier deleting.
+		IBuildInstallerThread* Thread;
+
 		// A flag storing whether the process was a success.
-		FThreadSafeBool bSuccess;
+		std::atomic<bool> bSuccess;
 
 		// A flag marking that we a running.
-		FThreadSafeBool bIsRunning;
+		std::atomic<bool> bIsRunning;
 
 		// A flag marking that we initialized correctly.
-		FThreadSafeBool bIsInited;
+		std::atomic<bool> bIsInited;
 
 		// A flag that stores whether we are on the first install iteration.
-		FThreadSafeBool bFirstInstallIteration;
+		std::atomic<bool> bFirstInstallIteration;
+
+		// A flag that indicates whether we are on the last install iteration.
+		std::atomic<bool> bLastInstallIteration = false;
 
 		// Tracks required download data between install retries.
-		FThreadSafeInt64 PreviousTotalDownloadRequired;
+		std::atomic<int64> PreviousTotalDownloadRequired;
+
+		// Tracks if any results of the file construction are encrypted. If they are, we'll be staging only, and no point verifying.
+		std::atomic<bool> bAreResultsEncrypted;
 
 		// Keep track of build stats.
 		FBuildInstallStats BuildStats;
@@ -145,17 +161,17 @@ namespace BuildPatchServices
 		// Where to put the files in the new installation.
 		TMap<FString, FString> NewPerFileSubdirectories;
 
-		// Map of registered installations.
-		TMultiMap<FString, FBuildPatchAppManifestRef> InstallationInfo;
-
 		// The file which contains per machine configuration information.
 		FString LocalMachineConfigFile;
 
 		// HTTP manager used to make requests for download data.
-		TUniquePtr<IHttpManager> HttpManager;
+		TSharedPtr<IHttpManager> HttpManager;
 
 		// File systems for classes requiring disk access.
-		TUniquePtr<IFileSystem> FileSystem;
+		TSharedPtr<IFileSystem> FileSystem;
+
+		// File systems for classes requiring crypto.
+		TUniquePtr<ICrypto> Crypto;
 
 		// Platform abstraction.
 		TUniquePtr<IPlatform> Platform;
@@ -167,7 +183,7 @@ namespace BuildPatchServices
 		TSharedPtr<IAnalyticsProvider> Analytics;
 
 		// Installer analytics handler.
-		TUniquePtr<IInstallerAnalytics> InstallerAnalytics;
+		TSharedPtr<IInstallerAnalytics> InstallerAnalytics;
 
 		// Installer statistics tracking.
 		TUniquePtr<IFileOperationTracker> FileOperationTracker;
@@ -175,7 +191,7 @@ namespace BuildPatchServices
 		TUniquePtr<ISpeedRecorder> DiskReadSpeedRecorder;
 		TUniquePtr<ISpeedRecorder> DiskWriteSpeedRecorder;
 		TUniquePtr<IChunkDataSizeProvider> ChunkDataSizeProvider;
-		TUniquePtr<IDownloadServiceStatistics> DownloadServiceStatistics;
+		TSharedPtr<IDownloadServiceStatistics> DownloadServiceStatistics;
 		TUniquePtr<IChunkDbChunkSourceStatistics> ChunkDbChunkSourceStatistics;
 		TUniquePtr<IInstallChunkSourceStatistics> InstallChunkSourceStatistics;
 		TUniquePtr<ICloudChunkSourceStatistics> CloudChunkSourceStatistics;
@@ -197,9 +213,22 @@ namespace BuildPatchServices
 		// List of controllable classes that have been constructed.
 		TArray<IControllable*> Controllables;
 
+		// Download rate limit in Bytes per second, 0 is unlimited.
+		// Since we can get a change on any thread and we may need to pipe it to
+		// the file constructor, rather than poll the rate we register a function
+		// that gets called. The function must be set within the mutex and will be called
+		// holding the mutex.
+		UE::FMutex DownloadRateLimitLock;
+		uint64 DownloadRateLimitBps;
+		TUniqueFunction<void(uint64 NewRate)> DownloadRateUpdatedFn;
+
+		// List of message handlers that have been registered.
+		TArray<FMessageHandler*> MessageHandlers;
+
 		// Stage timers for build stats.
 		FProcessTimer InitializeTimer;
 		FProcessTimer ConstructTimer;
+		FProcessTimer UninstallActionTimer;
 		FProcessTimer MoveFromStageTimer;
 		FProcessTimer FileAttributesTimer;
 		FProcessTimer VerifyTimer;
@@ -240,6 +269,9 @@ namespace BuildPatchServices
 		virtual bool IsPaused() const override;
 		virtual bool IsResumable() const override;
 		virtual bool IsUpdate() const override;
+		virtual bool IsRepair() const override;
+		virtual bool IsVerifyOnly() const override;
+		virtual bool IsUninstall() const override;
 		virtual bool CompletedSuccessfully() const override;
 		virtual bool HasError() const override;
 		virtual EBuildPatchInstallError GetErrorType() const override;
@@ -262,6 +294,8 @@ namespace BuildPatchServices
 		virtual void RegisterMessageHandler(FMessageHandler* MessageHandler) override;
 		virtual void UnregisterMessageHandler(FMessageHandler* MessageHandler) override;
 		virtual const FBuildInstallerConfiguration& GetConfiguration() const override;
+		virtual void SetDownloadBytesPerSecondLimit(uint64 InDownloadRateLimitBps) override;
+		virtual uint64 GetDownloadBytesPerSecondLimit() const override;
 #if !UE_BUILD_SHIPPING
 		virtual void GetDebugText(TArray<FString>& Output) override;
 #endif
@@ -362,11 +396,28 @@ namespace BuildPatchServices
 		 */
 		bool RunInstallation(TArray<FString>& CorruptFiles);
 
+		/**
+		 * Runs the preload download process.
+		 * @return  Returns true if there were no errors blocking the preload.
+		 */
+		bool RunPreload();
+
 		/**
 		 * Runs the prerequisite installation process.
 		 */
 		bool RunPrerequisites();
 
+		/**
+		 * Sets installation stats from various systems.
+		 */
+		void SetFinalInstallStats();
+
+		/**
+		 * Runs the custom uninstall action if specified.
+		 * @return    Returns true if there were no errors.
+		 */
+		bool RunCustomUninstallAction();
+
 		/**
 		 * Runs the backup process for locally changed files, and then moves new files into installation directory.
 		 * @return    Returns true if there were no errors.
@@ -428,11 +479,32 @@ namespace BuildPatchServices
 		 */
 		bool RelocateFileWithRetries(const FString& ToFullFilename, const FString& FromFullFilename, uint32& RenameErrorCode, uint32& CopyErrorCode);
 
+		/**
+		 * Attempt to copy a file using configured retries
+		 * @param ToFullFilename          The full filename of the destination.
+		 * @param FromFullFilename        The full filename of the source.
+		 * @param CopyErrorCode       OUT The copy error code if failed.
+		 * @return true if the file was copied.
+		 */
+		bool CopyFileWithRetries(const FString& ToFullFilename, const FString& FromFullFilename, uint32& CopyErrorCode);
+
+		/**
+		 * Builds the cloud source configuration struct.
+		 */
+		FCloudSourceConfig BuildCloudSourceConfig();
+
 		/**
 		 * Builds the cloud source configuration struct.
 		 */
 		FDownloadConnectionCountConfig BuildConnectionCountConfig();
 
+		/**
+		 * Get the filename as SHA string value.
+		 * @param FileManifest    The FFileManifest for the file.
+		 * @return filename as SHA string value.
+		 */
+		FString GetSHA1Filename(const FFileManifest& FileManifest);
+
 		/**
 		 * Destroys the internal worker thread
 		 */
diff --git a/Engine/Source/Runtime/Online/BuildPatchServices/Private/BuildPatchManifest.cpp b/Engine/Source/Runtime/Online/BuildPatchServices/Private/BuildPatchManifest.cpp
--- a/Engine/Source/Runtime/Online/BuildPatchServices/Private/BuildPatchManifest.cpp
+++ b/Engine/Source/Runtime/Online/BuildPatchServices/Private/BuildPatchManifest.cpp
@@ -2,6 +2,7 @@
 
 #include "BuildPatchManifest.h"
 #include "HAL/FileManager.h"
+#include "Misc/Base64.h"
 #include "Misc/FileHelper.h"
 #include "Misc/Paths.h"
 #include "Serialization/MemoryWriter.h"
@@ -14,11 +15,16 @@
 #include "Algo/Sort.h"
 #include "Algo/Accumulate.h"
 #include "Algo/Transform.h"
+#include "Algo/RemoveIf.h"
 #include "Core/BlockStructure.h"
+#include "Core/AsyncHelpers.h"
 #include "Data/ChunkData.h"
 #include "Data/ManifestData.h"
 #include "BuildPatchUtil.h"
 
+DECLARE_LOG_CATEGORY_EXTERN(LogBuildPatchManifest, Log, All);
+DEFINE_LOG_CATEGORY(LogBuildPatchManifest);
+
 using namespace BuildPatchServices;
 
 #define LOCTEXT_NAMESPACE "BuildPatchManifest"
@@ -75,6 +81,56 @@ bool BufferIsJsonManifest(const TArray<uint8>& DataInput)
 	return false;
 }
 
+FString GetFilename(const FFileManifest& FileManifest, const FString& RootDirectory, bool bIsStaged)
+{
+	static constexpr const TCHAR* EncryptedExt = TEXT(".e");
+	static constexpr const TCHAR* CompletedExt = TEXT(".c");
+	if (bIsStaged)
+	{
+		FString FileSHA1HashStringValue;
+		FBuildPatchUtils::SHAToBase32(FileManifest.SHA1Hash, FileSHA1HashStringValue);
+		return IFileManager::Get().FileExists(*(RootDirectory / FileSHA1HashStringValue + CompletedExt)) ? FileSHA1HashStringValue + CompletedExt : FileSHA1HashStringValue + EncryptedExt;
+	}
+	else
+	{
+		return FileManifest.Filename;
+	}
+}
+
+FString GetFilename(const FString& Filename, const FString& /*RootDirectory*/, bool bIsStaged)
+{
+	// Staged files not supported if using filenames not file manifests
+	check(bIsStaged == false);
+	return Filename;
+}
+
+template<typename Cont>
+TMap<FString, int64> GetFileSizes(const Cont& FilesData, const FString& RootDirectory, bool bIsStaged = false)
+{
+	TArray<TTuple<FString, FString, int64>> TargetFiles;
+	TargetFiles.Reserve(FilesData.Num());
+	for (const auto& File : FilesData)
+	{
+		TargetFiles.Add(TTuple<FString, FString, int64>(GetFilename(File, RootDirectory, false), GetFilename(File, RootDirectory, bIsStaged), 0));
+	}
+	auto GetFileSizesTasks = CreateTaskSchedulerLimitedByCoresNumber<int64>(EAsyncExecution::TaskGraph);
+	for (auto& FileInfo : TargetFiles)
+	{
+		GetFileSizesTasks.ScheduleTask(
+			[&RootDirectory, &FileInfo]() { return IFileManager::Get().FileSize(*(RootDirectory / FileInfo.Get<1>())); },
+			[&FileInfo](int64 FileSize) { FileInfo.Get<2>() = FileSize; }
+		);
+	}
+	GetFileSizesTasks.WaitAll();
+	TMap<FString, int64> Results;
+	for (auto& FileInfo : TargetFiles)
+	{
+		Results.FindOrAdd(FileInfo.Get<0>(), FileInfo.Get<2>());
+	}
+	return Results;
+}
+
+
 /* FBuildPatchCustomField implementation
 *****************************************************************************/
 FBuildPatchCustomField::FBuildPatchCustomField(const FString& Value)
@@ -109,6 +165,50 @@ int64 FBuildPatchCustomField::AsInteger() const
 	return 0;
 }
 
+/* FCipherHeader implementation
+*****************************************************************************/
+struct FCipherHeader
+{
+public:
+	EFeatureLevel Version = EFeatureLevel::Latest;
+	uint32 HeaderSize = 0;
+	EManifestStorageFlags DataStoredAs = EManifestStorageFlags::RawData;
+	uint32 DataSizeUncompressed = 0;
+	uint32 DataSizeCompressed = 0;
+	TArray<uint8> InitializationVector;
+};
+
+FArchive& operator<<(FArchive& Ar, FCipherHeader& Header)
+{
+	if (Ar.IsError())
+	{
+		return Ar;
+	}
+
+	const int64 StartPos = Ar.Tell();
+	Ar << Header.HeaderSize;
+
+	int32 Version = (int32)Header.Version;
+	Ar << Version;
+	Header.Version = (EFeatureLevel)Version;
+
+	Ar << Header.DataStoredAs;
+	Ar << Header.DataSizeUncompressed;
+	Ar << Header.DataSizeCompressed;
+	Ar << Header.InitializationVector;
+
+	if (Ar.IsSaving())
+	{
+		const int64 EndPos = Ar.Tell();
+		Ar.Seek(StartPos);
+		Ar << (Header.HeaderSize = EndPos - StartPos);
+		Ar.Seek(EndPos);
+	}
+
+	Ar.Seek(StartPos + Header.HeaderSize);
+	return Ar;
+}
+
 /* FBuildPatchAppManifest implementation
 *****************************************************************************/
 
@@ -131,10 +231,14 @@ FBuildPatchAppManifest::FBuildPatchAppManifest(const FBuildPatchAppManifest& Oth
 	, ChunkDataList(Other.ChunkDataList)
 	, FileManifestList(Other.FileManifestList)
 	, CustomFields(Other.CustomFields)
+	, EncryptedData(Other.EncryptedData)
+	, EncryptionSecretId(Other.EncryptionSecretId)
+	, bIsManifestEncrypted(Other.bIsManifestEncrypted)
 	, TotalBuildSize(Other.TotalBuildSize)
 	, TotalDownloadSize(Other.TotalDownloadSize)
 	, bNeedsResaving(Other.bNeedsResaving)
 {
+	FMemory::Memcpy(EncryptionAuthTag.AuthTag, Other.EncryptionAuthTag.AuthTag, AES256_GCM_AuthTagSizeInBytes);
 	InitLookups();
 }
 
@@ -143,27 +247,43 @@ FBuildPatchAppManifest::~FBuildPatchAppManifest()
 	DestroyData();
 }
 
-bool FBuildPatchAppManifest::SaveToFile(const FString& Filename, BuildPatchServices::EFeatureLevel SaveFormat)
+bool FBuildPatchAppManifest::SaveToFile(const FString& Filename, BuildPatchServices::EFeatureLevel SaveFormat, FSHAHash* OutSHA1Hash, FMD5Hash* OutMD5Hash)
 {
 	bool bSuccess = SaveFormat >= GetFeatureLevel();
 	if (bSuccess)
 	{
-		TUniquePtr<FArchive> FileOut(IFileManager::Get().CreateFileWriter(*Filename));
-		bSuccess = FileOut.IsValid();
+		TArray<uint8> ManifestData;
+		FMemoryWriter MemoryWriter(ManifestData);
+		// Serialize into the desired format.
+		if (SaveFormat >= BuildPatchServices::EFeatureLevel::StoredAsBinaryData)
+		{
+			bSuccess = FManifestData::Serialize(MemoryWriter, *this, SaveFormat);
+		}
+		else
+		{
+			FString JSONOutput;
+			SerializeToJSON(JSONOutput);
+			FTCHARToUTF8 JsonUTF8(*JSONOutput);
+			MemoryWriter.Serialize((UTF8CHAR*)JsonUTF8.Get(), JsonUTF8.Length() * sizeof(UTF8CHAR));
+		}
+		// Save the file if successful.
+		if (bSuccess)
+		{
+			bSuccess = FFileHelper::SaveArrayToFile(ManifestData, *Filename);
+		}
+		// Set the output hash values if desired.
 		if (bSuccess)
 		{
-			if (SaveFormat >= BuildPatchServices::EFeatureLevel::StoredAsBinaryData)
+			if (OutMD5Hash != nullptr)
 			{
-				bSuccess = FManifestData::Serialize(*FileOut, *this, SaveFormat);
+				FMD5 Md5Gen;
+				Md5Gen.Update(ManifestData.GetData(), ManifestData.Num());
+				OutMD5Hash->Set(Md5Gen);
 			}
-			else
+			if (OutSHA1Hash != nullptr)
 			{
-				FString JSONOutput;
-				SerializeToJSON(JSONOutput);
-				FTCHARToUTF8 JsonUTF8(*JSONOutput);
-				FileOut->Serialize((UTF8CHAR*)JsonUTF8.Get(), JsonUTF8.Length() * sizeof(UTF8CHAR));
+				FSHA1::HashBuffer(ManifestData.GetData(), ManifestData.Num(), OutSHA1Hash->Hash);
 			}
-			bSuccess = FileOut->Close() && bSuccess;
 		}
 	}
 	return bSuccess;
@@ -172,7 +292,7 @@ bool FBuildPatchAppManifest::SaveToFile(const FString& Filename, BuildPatchServi
 bool FBuildPatchAppManifest::LoadFromFile(const FString& Filename)
 {
 	TArray<uint8> FileData;
-	if (FFileHelper::LoadFileToArray(FileData, *Filename))
+	if (FFileHelper::LoadFileToArray(FileData, *Filename, FILEREAD_Silent))
 	{
 		return DeserializeFromData(FileData);
 	}
@@ -205,6 +325,7 @@ void FBuildPatchAppManifest::DestroyData()
 	ChunkDataList = FChunkDataList();
 	FileManifestList = FFileManifestList();
 	CustomFields = FCustomFields();
+	EncryptionSecretId = FGuid();
 	FileNameLookup.Empty();
 	FileManifestLookup.Empty();
 	TaggedFilesLookup.Empty();
@@ -297,7 +418,7 @@ void FBuildPatchAppManifest::SerializeToJSON(FString& JSONOutput)
 			Writer->WriteObjectStart();
 			{
 				Writer->WriteValue(TEXT("Filename"), FileManifest.Filename);
-				Writer->WriteValue(TEXT("FileHash"), FString::FromBlob(FileManifest.FileHash.Hash, FSHA1::DigestSize));
+				Writer->WriteValue(TEXT("FileHash"), FString::FromBlob(FileManifest.SHA1Hash.Hash, FSHA1::DigestSize));
 				if (EnumHasAllFlags(FileManifest.FileMetaFlags, EFileMetaFlags::UnixExecutable))
 				{
 					Writer->WriteValue(TEXT("bIsUnixExecutable"), true);
@@ -491,7 +612,7 @@ bool FBuildPatchAppManifest::DeserializeFromJSON( const FString& JSONInput )
 			const int32 FileIndex = FileManifestList.FileList.Add(FFileManifest());
 			FFileManifest& FileManifest = FileManifestList.FileList[FileIndex];
 			FileManifest.Filename = JsonFileManifest->GetStringField(TEXT("Filename"));
-			bSuccess = bSuccess && FString::ToBlob(JsonFileManifest->GetStringField(TEXT("FileHash")), FileManifest.FileHash.Hash, FSHA1::DigestSize);
+			bSuccess = bSuccess && FString::ToBlob(JsonFileManifest->GetStringField(TEXT("FileHash")), FileManifest.SHA1Hash.Hash, FSHA1::DigestSize);
 			TArray<TSharedPtr<FJsonValue>> JsonChunkPartArray = JsonFileManifest->GetArrayField(TEXT("FileChunkParts"));
 			for (auto JsonChunkPartIt = JsonChunkPartArray.CreateConstIterator(); JsonChunkPartIt && bSuccess; ++JsonChunkPartIt)
 			{
@@ -606,7 +727,7 @@ bool FBuildPatchAppManifest::DeserializeFromJSON( const FString& JSONInput )
 		if (FoundFileManifest)
 		{
 			FSHAHash PrereqHash;
-			FMemory::Memcpy(PrereqHash.Hash, (*FoundFileManifest)->FileHash.Hash, FSHA1::DigestSize);
+			FMemory::Memcpy(PrereqHash.Hash, (*FoundFileManifest)->SHA1Hash.Hash, FSHA1::DigestSize);
 			ManifestMeta.PrereqIds.Add(PrereqHash.ToString());
 		}
 	}
@@ -707,7 +828,7 @@ bool FBuildPatchAppManifest::DeserializeFromJSON( const FString& JSONInput )
 				{
 					FChunkInfo* ChunkInfoData = MutableChunkInfoLookup[Guid];
 					ChunkInfoData->FileSize = FileManifest.FileSize;
-					ChunkInfoData->ShaHash = FileManifest.FileHash;
+					ChunkInfoData->ShaHash = FileManifest.SHA1Hash;
 				}
 			}
 			else
@@ -720,6 +841,11 @@ bool FBuildPatchAppManifest::DeserializeFromJSON( const FString& JSONInput )
 	// Setup build id from backwards compat route
 	ManifestMeta.BuildId = FBuildPatchUtils::GetBackwardsCompatibleBuildId(ManifestMeta);
 
+	// Copy FeatureLevel to each complex serialisable member
+	ChunkDataList.FeatureLevel = ManifestMeta.FeatureLevel;
+	FileManifestList.FeatureLevel = ManifestMeta.FeatureLevel;
+	CustomFields.FeatureLevel = ManifestMeta.FeatureLevel;
+
 	// Call OnPostLoad for the file manifest list
 	FileManifestList.OnPostLoad();
 
@@ -738,11 +864,6 @@ bool FBuildPatchAppManifest::DeserializeFromJSON( const FString& JSONInput )
 	return bSuccess;
 }
 
-EFeatureLevel FBuildPatchAppManifest::GetFeatureLevel() const
-{
-	return ManifestMeta.FeatureLevel;
-}
-
 void FBuildPatchAppManifest::GetChunksRequiredForFiles(const TSet<FString>& Filenames, TSet<FGuid>& RequiredChunks) const
 {
 	for (const FString& Filename : Filenames)
@@ -765,28 +886,26 @@ int64 FBuildPatchAppManifest::GetDownloadSize() const
 
 int64 FBuildPatchAppManifest::GetDownloadSize(const TSet<FString>& Tags) const
 {
-	// For each tag we iterate the files and for each new chunk we find we add the download size for it.
+	return GetDownloadSize(GetTaggedFileManifests(Tags));
+}
+
+int64 FBuildPatchAppManifest::GetDownloadSize(const TSet<const FFileManifest*>& TaggedFiles) const
+{
+	// For each tagged file and for each new chunk we find we add the download size for it.
 	TSet<FGuid> RequiredChunks;
 	int64 TotalSize = 0;
-	for (const FString& Tag : Tags)
+	for (const FFileManifest* File : TaggedFiles)
 	{
-		const TArray<const FFileManifest*>* Files = TaggedFilesLookup.Find(Tag);
-		if (Files != nullptr)
+		for (const FChunkPart& ChunkPart : File->ChunkParts)
 		{
-			for (const FFileManifest* File : *Files)
+			bool bAlreadyInSet;
+			RequiredChunks.Add(ChunkPart.Guid, &bAlreadyInSet);
+			if (!bAlreadyInSet)
 			{
-				for (const FChunkPart& ChunkPart : File->ChunkParts)
+				const FChunkInfo* const* ChunkInfo = ChunkInfoLookup.Find(ChunkPart.Guid);
+				if (ChunkInfo != nullptr)
 				{
-					bool bAlreadyInSet;
-					RequiredChunks.Add(ChunkPart.Guid, &bAlreadyInSet);
-					if (!bAlreadyInSet)
-					{
-						const FChunkInfo * const * ChunkInfo = ChunkInfoLookup.Find(ChunkPart.Guid);
-						if (ChunkInfo != nullptr)
-						{
-							TotalSize += (*ChunkInfo)->FileSize;
-						}
-					}
+					TotalSize += (*ChunkInfo)->FileSize;
 				}
 			}
 		}
@@ -801,38 +920,61 @@ int64 FBuildPatchAppManifest::GetDeltaDownloadSize(const TSet<FString>& Tags, co
 
 int64 FBuildPatchAppManifest::GetDeltaDownloadSize(const TSet<FString>& InTags, const IBuildManifestRef& InPreviousVersion, const TSet<FString>& InPreviousTags) const
 {
-	TSet<FString> Tags = InTags;
+	const TSet<FString>* Tags = &InTags;
 	FBuildPatchAppManifestRef PreviousVersion = StaticCastSharedRef< FBuildPatchAppManifest >(InPreviousVersion);
-	TSet<FString> PreviousTags = InPreviousTags;
-	if (Tags.Num() == 0)
+	const TSet<FString>* PreviousTags = &InPreviousTags;
+	TSet<FString> TagsManifest, PreviousTagsManifest;
+	if (Tags->Num() == 0)
 	{
-		GetFileTagList(Tags);
+		GetFileTagList(TagsManifest);
+		Tags = &TagsManifest;
 	}
-	if (PreviousTags.Num() == 0)
+	if (PreviousTags->Num() == 0)
 	{
-		PreviousVersion->GetFileTagList(PreviousTags);
+		PreviousVersion->GetFileTagList(PreviousTagsManifest);
+		PreviousTags = &PreviousTagsManifest;
 	}
 
 	// Enumerate what is available.
 	TSet<FString> FilesInstalled;
 	TSet<FGuid> ChunksRequired;
 	TSet<FGuid> ChunksInstalled;
-	PreviousVersion->GetTaggedFileList(PreviousTags, FilesInstalled);
+	PreviousVersion->GetTaggedFileList(*PreviousTags, FilesInstalled);
 	PreviousVersion->GetChunksRequiredForFiles(FilesInstalled, ChunksRequired);
-	PreviousVersion->EnumerateProducibleChunks(PreviousTags, ChunksRequired, ChunksInstalled);
+	PreviousVersion->EnumerateProducibleChunks(*PreviousTags, ChunksRequired, ChunksInstalled);
+
+	// Enumerate what is needed for the update.
+	TSet<FString> AllTaggedFiles;
+	GetTaggedFileList(*Tags, AllTaggedFiles);
+
+	return GetDeltaDownloadSize(AllTaggedFiles, InPreviousVersion, FilesInstalled, ChunksInstalled);
+}
+
+int64 FBuildPatchAppManifest::GetDeltaDownloadSize(const TSet<const FFileManifest*>& InTaggedManifests, const IBuildManifestRef& InPreviousVersion, const TSet<FString>& InFilesInstalled, const TSet<FGuid>& InChunksInstalled) const
+{
+	// Enumerate what is needed for the update.
+	TSet<FString> AllTaggedFiles;
+	Algo::Transform(InTaggedManifests, AllTaggedFiles, [](const FFileManifest* File) { return File->Filename; });
+
+	return GetDeltaDownloadSize(AllTaggedFiles, InPreviousVersion, InFilesInstalled, InChunksInstalled);
+}
+
+int64 FBuildPatchAppManifest::GetDeltaDownloadSize(const TSet<FString>& InAllTaggedFiles, const IBuildManifestRef& InPreviousVersion, const TSet<FString>& InFilesInstalled, const TSet<FGuid>& InChunksInstalled) const
+{
+	FBuildPatchAppManifestRef PreviousVersion = StaticCastSharedRef< FBuildPatchAppManifest >(InPreviousVersion);
 
 	// Enumerate what has changed.
 	FString DummyString;
 	TSet<FString> OutdatedFiles;
-	GetOutdatedFiles(PreviousVersion, DummyString, OutdatedFiles);
+	GetOutdatedFiles(&PreviousVersion.Get(), DummyString, InAllTaggedFiles, OutdatedFiles);
+
+	const TSet<FString> NewFilesNeeded = InAllTaggedFiles.Difference(InFilesInstalled);
+	const TSet<FString> UpdatedFilesNeeded = OutdatedFiles.Intersect(InAllTaggedFiles);
+	const TSet<FString> FilesNeeded = NewFilesNeeded.Union(UpdatedFilesNeeded);
 
-	// Enumerate what is needed for the update.
-	TSet<FString> FilesNeeded;
 	TSet<FGuid> ChunksNeeded;
-	GetTaggedFileList(Tags, FilesNeeded);
-	FilesNeeded = OutdatedFiles.Intersect(FilesNeeded);
 	GetChunksRequiredForFiles(FilesNeeded, ChunksNeeded);
-	ChunksNeeded = ChunksNeeded.Difference(ChunksInstalled);
+	ChunksNeeded = ChunksNeeded.Difference(InChunksInstalled);
 
 	// Return download size of required chunks.
 	return GetDataSize(ChunksNeeded);
@@ -845,26 +987,13 @@ int64 FBuildPatchAppManifest::GetBuildSize() const
 
 int64 FBuildPatchAppManifest::GetBuildSize(const TSet<FString>& Tags) const
 {
-	// For each tag we iterate the files and for each new file we find we add the size for it.
-	TSet<const FFileManifest*> RequiredFiles;
-	int64 TotalSize = 0;
-	for (const FString& Tag : Tags)
-	{
-		const TArray<const FFileManifest*>* Files = TaggedFilesLookup.Find(Tag);
-		if (Files != nullptr)
-		{
-			for (const FFileManifest* File : *Files)
-			{
-				bool bAlreadyInSet;
-				RequiredFiles.Add(File, &bAlreadyInSet);
-				if (!bAlreadyInSet)
-				{
-					TotalSize += File->FileSize;
-				}
-			}
-		}
-	}
-	return TotalSize;
+	return GetBuildSize(GetTaggedFileManifests(Tags));
+}
+
+int64 FBuildPatchAppManifest::GetBuildSize(const TSet<const FFileManifest*>& TaggedFiles) const
+{
+	// For each tag file and for each new file we find we add the size for it.
+	return Algo::Accumulate<int64>(TaggedFiles, 0, [](int64 Result, const FFileManifest* File) { return Result + File->FileSize; });
 }
 
 TArray<FString> FBuildPatchAppManifest::GetBuildFileList() const
@@ -908,13 +1037,72 @@ int64 FBuildPatchAppManifest::GetFileSize(const TSet<FString>& Filenames) const
 int64 FBuildPatchAppManifest::GetFileSize(FStringView Filename) const
 {
 	const FFileManifest *const *const FileManifest = FileManifestLookup.Find(Filename);
-	if (FileManifest)
+	if (FileManifest != nullptr)
 	{
 		return (*FileManifest)->FileSize;
 	}
 	return 0;
 }
 
+EFileMetaFlags FBuildPatchAppManifest::GetFileMetaFlags(const FString& Filename) const
+{
+	const FFileManifest *const *const FileManifest = FileManifestLookup.Find(Filename);
+	if (FileManifest != nullptr)
+	{
+		return (*FileManifest)->FileMetaFlags;
+	}
+	return EFileMetaFlags::None;
+}
+
+const FMD5Hash* FBuildPatchAppManifest::GetFileMD5Hash(const FString& Filename) const
+{
+	if (GetFeatureLevel() >= BuildPatchServices::EFeatureLevel::StoresFileMD5HashesAndMIMEType)
+	{
+		const FFileManifest* const* const FileManifest = FileManifestLookup.Find(Filename);
+		if (FileManifest != nullptr)
+		{
+			return &(*FileManifest)->MD5Hash;
+		}
+	}
+	return nullptr;
+}
+
+const FSHAHash* FBuildPatchAppManifest::GetFileSHA1Hash(const FString& Filename) const
+{
+	const FFileManifest* const* const FileManifest = FileManifestLookup.Find(Filename);
+	if (FileManifest != nullptr)
+	{
+		return &(*FileManifest)->SHA1Hash;
+	}
+	return nullptr;
+}
+
+const FSHA256Signature* FBuildPatchAppManifest::GetFileSHA256Hash(const FString& Filename) const
+{
+	if (GetFeatureLevel() >= BuildPatchServices::EFeatureLevel::StoresFileSHA256Hashes)
+	{
+		const FFileManifest* const* const FileManifest = FileManifestLookup.Find(Filename);
+		if (FileManifest != nullptr)
+		{
+			return &(*FileManifest)->SHA256Hash;
+		}
+	}
+	return nullptr;
+}
+
+const FString* FBuildPatchAppManifest::GetFileMIMEType(const FString& Filename) const
+{
+	if (GetFeatureLevel() >= BuildPatchServices::EFeatureLevel::StoresFileMD5HashesAndMIMEType)
+	{
+		const FFileManifest* const* const FileManifest = FileManifestLookup.Find(Filename);
+		if (FileManifest != nullptr)
+		{
+			return &(*FileManifest)->MIMEType;
+		}
+	}
+	return nullptr;
+}
+
 int64 FBuildPatchAppManifest::GetDataSize(const FGuid& DataGuid) const
 {
 	if (ChunkInfoLookup.Contains(DataGuid))
@@ -979,17 +1167,7 @@ void FBuildPatchAppManifest::GetFileTagList(TSet<FString>& Tags) const
 
 void FBuildPatchAppManifest::GetTaggedFileList(const TSet<FString>& Tags, TArray<FString>& TaggedFiles) const
 {
-	for (const FString& Tag : Tags)
-	{
-		const TArray<const FFileManifest*> *const Files = TaggedFilesLookup.Find(Tag);
-		if (Files != nullptr)
-		{
-			for (const FFileManifest* File : *Files)
-			{
-				TaggedFiles.Add(File->Filename);
-			}
-		}
-	}
+	Algo::Transform(GetTaggedFileManifests(Tags), TaggedFiles, [&](const FFileManifest* FileManifest) { return FileManifest->Filename; });
 }
 
 void FBuildPatchAppManifest::GetTaggedFileList(const TSet<FString>& Tags, TArray<FStringView>& TaggedFiles) const
@@ -1009,17 +1187,20 @@ void FBuildPatchAppManifest::GetTaggedFileList(const TSet<FString>& Tags, TArray
 
 void FBuildPatchAppManifest::GetTaggedFileList(const TSet<FString>& Tags, TSet<FString>& TaggedFiles) const
 {
+	Algo::Transform(GetTaggedFileManifests(Tags), TaggedFiles, [&](const FFileManifest* FileManifest) { return FileManifest->Filename; });
+}
+
+TSet<const FFileManifest*> FBuildPatchAppManifest::GetTaggedFileManifests(const TSet<FString>& Tags) const
+{
+	TSet<const FFileManifest*> TaggedFiles;
 	for (const FString& Tag : Tags)
 	{
-		const TArray<const FFileManifest*> *const Files = TaggedFilesLookup.Find(Tag);
-		if (Files != nullptr)
+		if (const TArray<const FFileManifest*>* const Files = TaggedFilesLookup.Find(Tag))
 		{
-			for (const FFileManifest* File : *Files)
-			{
-				TaggedFiles.Add(File->Filename);
-			}
+			TaggedFiles.Append(*Files);
 		}
 	}
+	return TaggedFiles;
 }
 
 void FBuildPatchAppManifest::GetDataList(TArray<FGuid>& DataGuids) const
@@ -1035,7 +1216,7 @@ void FBuildPatchAppManifest::GetDataList(TSet<FGuid>& DataGuids) const
 const FFileManifest* FBuildPatchAppManifest::GetFileManifest(const FString& Filename) const
 {
 	const FFileManifest* const * FileManifest = FileManifestLookup.Find(Filename);
-	return (FileManifest) ? (*FileManifest) : nullptr;
+	return (FileManifest != nullptr) ? (*FileManifest) : nullptr;
 }
 
 bool FBuildPatchAppManifest::IsFileDataManifest() const
@@ -1046,7 +1227,7 @@ bool FBuildPatchAppManifest::IsFileDataManifest() const
 bool FBuildPatchAppManifest::GetChunkHash(const FGuid& ChunkGuid, uint64& OutHash) const
 {
 	const FChunkInfo* const * ChunkInfo = ChunkInfoLookup.Find(ChunkGuid);
-	if (ChunkInfo)
+	if (ChunkInfo != nullptr)
 	{
 		OutHash = (*ChunkInfo)->Hash;
 		return true;
@@ -1066,6 +1247,17 @@ bool FBuildPatchAppManifest::GetChunkShaHash(const FGuid& ChunkGuid, FSHAHash& O
 	return false;
 }
 
+bool FBuildPatchAppManifest::GetChunkEncryptionSecretId(const FGuid& ChunkGuid, FGuid& OutEncryptionSecretId) const
+{
+	const FChunkInfo* const* ChunkInfo = ChunkInfoLookup.Find(ChunkGuid);
+	if (ChunkInfo != nullptr)
+	{
+		OutEncryptionSecretId = (*ChunkInfo)->EncryptionSecretId;
+		return OutEncryptionSecretId.IsValid();
+	}
+	return false;
+}
+
 const BuildPatchServices::FChunkInfo* FBuildPatchAppManifest::GetChunkInfo(const FGuid& ChunkGuid) const
 {
 	const BuildPatchServices::FChunkInfo* const * const ChunkInfoPtrPtr = ChunkInfoLookup.Find(ChunkGuid);
@@ -1076,7 +1268,7 @@ const BuildPatchServices::FChunkInfo* FBuildPatchAppManifest::GetChunkInfo(const
 bool FBuildPatchAppManifest::GetFileHash(const FGuid& FileGuid, FSHAHash& OutHash) const
 {
 	const FString* const * FoundFilename = FileNameLookup.Find(FileGuid);
-	if (FoundFilename)
+	if (FoundFilename != nullptr)
 	{
 		return GetFileHash(**FoundFilename, OutHash);
 	}
@@ -1086,9 +1278,9 @@ bool FBuildPatchAppManifest::GetFileHash(const FGuid& FileGuid, FSHAHash& OutHas
 bool FBuildPatchAppManifest::GetFileHash(const FString& Filename, FSHAHash& OutHash) const
 {
 	const FFileManifest* const * FoundFileManifest = FileManifestLookup.Find(Filename);
-	if (FoundFileManifest)
+	if (FoundFileManifest != nullptr)
 	{
-		FMemory::Memcpy(OutHash.Hash, (*FoundFileManifest)->FileHash.Hash, FSHA1::DigestSize);
+		FMemory::Memcpy(OutHash.Hash, (*FoundFileManifest)->SHA1Hash.Hash, FSHA1::DigestSize);
 		return true;
 	}
 	return false;
@@ -1097,7 +1289,7 @@ bool FBuildPatchAppManifest::GetFileHash(const FString& Filename, FSHAHash& OutH
 bool FBuildPatchAppManifest::GetFilePartHash(const FGuid& FilePartGuid, uint64& OutHash) const
 {
 	const FChunkInfo* const * FilePartInfo = ChunkInfoLookup.Find(FilePartGuid);
-	if (FilePartInfo)
+	if (FilePartInfo != nullptr)
 	{
 		OutHash = (*FilePartInfo)->Hash;
 		return true;
@@ -1105,6 +1297,11 @@ bool FBuildPatchAppManifest::GetFilePartHash(const FGuid& FilePartGuid, uint64&
 	return false;
 }
 
+BuildPatchServices::EFeatureLevel FBuildPatchAppManifest::GetFeatureLevel() const
+{
+	return ManifestMeta.FeatureLevel;
+}
+
 uint32 FBuildPatchAppManifest::GetAppID() const
 {
 	return ManifestMeta.AppID;
@@ -1120,6 +1317,34 @@ const FString& FBuildPatchAppManifest::GetVersionString() const
 	return ManifestMeta.BuildVersion;
 }
 
+const FString& FBuildPatchAppManifest::GetUniqueBuildId() const
+{
+	return ManifestMeta.BuildId;
+}
+
+const FGuid& FBuildPatchAppManifest::GetEncryptionSecretId(bool* bOutIsManifestEncrypted) const
+{
+	if (bOutIsManifestEncrypted)
+	{
+		*bOutIsManifestEncrypted = bIsManifestEncrypted;
+	}
+	return EncryptionSecretId;
+}
+
+bool FBuildPatchAppManifest::IsManifestEncrypted() const
+{
+	return bIsManifestEncrypted;
+}
+
+TSet<FGuid> FBuildPatchAppManifest::GetNecessaryEncryptionSecretIds() const
+{
+	TSet<FGuid> EncryptionSecretIds;
+	EncryptionSecretIds.Add(EncryptionSecretId);
+	Algo::Transform(ChunkDataList.ChunkList, EncryptionSecretIds, &FChunkInfo::EncryptionSecretId);
+	EncryptionSecretIds.Remove(FGuid());
+	return EncryptionSecretIds;
+}
+
 const FString& FBuildPatchAppManifest::GetLaunchExe() const
 {
 	return ManifestMeta.LaunchExe;
@@ -1150,6 +1375,16 @@ const FString& FBuildPatchAppManifest::GetPrereqArgs() const
 	return ManifestMeta.PrereqArgs;
 }
 
+const FString& FBuildPatchAppManifest::GetUninstallActionPath() const
+{
+	return ManifestMeta.UninstallActionPath;
+}
+
+const FString& FBuildPatchAppManifest::GetUninstallActionArgs() const
+{
+	return ManifestMeta.UninstallActionArgs;
+}
+
 const FString& FBuildPatchAppManifest::GetBuildId() const
 {
 	return ManifestMeta.BuildId;
@@ -1174,6 +1409,13 @@ void FBuildPatchAppManifest::CopyCustomFields(const IBuildManifestRef& InOther,
 	}
 }
 
+TArray<FString> FBuildPatchAppManifest::GetCustomFieldNames() const
+{
+	TArray<FString> CustomFieldNames;
+	CustomFields.Fields.GenerateKeyArray(CustomFieldNames);
+	return CustomFieldNames;
+}
+
 const IManifestFieldPtr FBuildPatchAppManifest::GetCustomField(const FString& FieldName) const
 {
 	IManifestFieldPtr Return;
@@ -1213,11 +1455,19 @@ int32 FBuildPatchAppManifest::EnumerateProducibleChunks(const TSet<FString>& Tag
 	return EnumerateProducibleChunks_Internal([&](const FString& Filename) { return AvailableFiles.Contains(Filename); }, ChunksRequired, ChunksAvailable);
 }
 
+int32 FBuildPatchAppManifest::EnumerateProducibleChunks(const TSet<const FFileManifest*>& TaggedFiles, const TSet<FGuid>& ChunksRequired, TSet<FGuid>& ChunksAvailable) const
+{
+	TSet<FString> AvailableFiles;
+	Algo::Transform(TaggedFiles, AvailableFiles, [](const FFileManifest* File) { return File->Filename; });
+	return EnumerateProducibleChunks_Internal([&](const FString& Filename) { return AvailableFiles.Contains(Filename); }, ChunksRequired, ChunksAvailable);
+}
+
 int32 FBuildPatchAppManifest::EnumerateProducibleChunks_Internal(const TFunction<bool(const FString&)>& FileAccessChecker, const TSet<FGuid>& ChunksRequired, TSet<FGuid>& ChunksAvailable) const
 {
 	int32 Count = 0;
 	TMap<FGuid, FBlockStructure> ChunkBlockStructures;
 	ChunkBlockStructures.Reserve(ChunksRequired.Num());
+	// Re-assemble the information about chunk available blocks from *.e or *.c files. 
 	for (const FFileManifest& FileManifest : FileManifestList.FileList)
 	{
 		const bool bHasFile = FileAccessChecker(FileManifest.Filename);
@@ -1238,7 +1488,9 @@ int32 FBuildPatchAppManifest::EnumerateProducibleChunks_Internal(const TFunction
 		const FGuid& ChunkBlockId = ChunkBlockStructurePair.Key;
 		const FBlockStructure& ChunkBlocks = ChunkBlockStructurePair.Value;
 		const FChunkInfo* ChunkInfoPtr = ChunkInfoLookup[ChunkBlockId];
-		if (ChunkBlocks.GetHead() == ChunkBlocks.GetTail() && ChunkBlocks.GetHead() && ChunkBlocks.GetHead()->GetSize() == ChunkInfoPtr->WindowSize)
+		// A chunk is treated as fully uploaded if there are no gaps between the Head and Tail, 
+		// and the total block size is equal to the total expected chunk size.
+		if (ChunkBlocks.GetHead() == ChunkBlocks.GetTail() && ChunkBlocks.GetHead() && ChunkBlocks.GetHead()->GetSize() == ChunkInfoPtr->DataSizeUncompressed)
 		{
 			ChunksAvailable.Add(ChunkBlockId);
 			++Count;
@@ -1260,6 +1512,168 @@ bool FBuildPatchAppManifest::HasFileAttributes() const
 	return false;
 }
 
+bool FBuildPatchAppManifest::SerialiseSensitiveFields(FArchive& Ar)
+{
+	// If the metadata is changed, we need to consider those changes here and work by version.
+	// If the new featurelevel does not add fields that require encryption, just bump the assert here.
+	static_assert((int32)EFeatureLevel::Latest == 24, "Please check the below code is still of for new feature level, and update this assert.");
+
+	// This is the list of metadata fields that get encrypted.
+	Ar << ManifestMeta.LaunchExe;
+	Ar << ManifestMeta.LaunchCommand;
+	Ar << ManifestMeta.PrereqIds;
+	Ar << ManifestMeta.PrereqName;
+	Ar << ManifestMeta.PrereqPath;
+	Ar << ManifestMeta.PrereqArgs;
+	Ar << ManifestMeta.UninstallActionPath;
+	Ar << ManifestMeta.UninstallActionArgs;
+
+	// Next, serialise all of the potentially sensitive FileManifestList fields.
+	// Unfortunately, we cannot be const correct here due to FArchive API.
+	for (FFileManifest& FileManifest : FileManifestList.FileList)
+	{
+		Ar << FileManifest.Filename;
+		Ar << FileManifest.SymlinkTarget;
+	}
+
+	return !Ar.IsError();
+}
+
+void FBuildPatchAppManifest::ObfuscateSensitiveFields()
+{
+	// If the metadata is changed, we need to consider those changes here and work by version.
+	// If the new featurelevel does not add fields that require encryption, just bump the assert here.
+	static_assert((int32)EFeatureLevel::Latest == 24, "Please check the below code is still of for new feature level, and update this assert.");
+
+	ManifestMeta.LaunchExe = ManifestDataHelpers::ObfuscateString(ManifestMeta.LaunchExe);
+	ManifestMeta.LaunchCommand = ManifestDataHelpers::ObfuscateString(ManifestMeta.LaunchCommand);
+	ManifestMeta.PrereqIds = ManifestDataHelpers::ObfuscateStrings(ManifestMeta.PrereqIds);
+	ManifestMeta.PrereqName = ManifestDataHelpers::ObfuscateString(ManifestMeta.PrereqName);
+	ManifestMeta.PrereqPath = ManifestDataHelpers::ObfuscateString(ManifestMeta.PrereqPath);
+	ManifestMeta.PrereqArgs = ManifestDataHelpers::ObfuscateString(ManifestMeta.PrereqArgs);
+	ManifestMeta.UninstallActionPath = ManifestDataHelpers::ObfuscateString(ManifestMeta.UninstallActionPath);
+	ManifestMeta.UninstallActionArgs = ManifestDataHelpers::ObfuscateString(ManifestMeta.UninstallActionArgs);
+	for (FFileManifest& FileManifest : FileManifestList.FileList)
+	{
+		FileManifest.Filename = ManifestDataHelpers::ObfuscateString(FileManifest.Filename);
+		FileManifest.SymlinkTarget = ManifestDataHelpers::ObfuscateString(FileManifest.SymlinkTarget);
+	}
+}
+
+bool FBuildPatchAppManifest::EncryptData(const FGuid& InEncryptionSecretId, const TArray<uint8>& InEncryptionSecretKey)
+{
+	using namespace BuildPatchServices;
+	checkf(!bIsManifestEncrypted, TEXT("Cannot encrypt already encrypted manifest file."));
+	checkf(GetFeatureLevel() >= EFeatureLevel::ManifestEncryptionSupport, TEXT("Cannot encrypt a manifest unless the support is new enough"));
+
+	TArray<uint8> DataToEncrypt;
+	FMemoryWriter RawAr(DataToEncrypt);
+	if (SerialiseSensitiveFields(RawAr))
+	{
+		// Now compress,because usually the sensitive manifest string fields would compress well, but if we encrypt without compression, we risk increasing the manifest file size dramatically.
+		FCipherHeader CipherHeader;
+		CipherHeader.DataSizeUncompressed = DataToEncrypt.Num();
+		CipherHeader.DataStoredAs = ManifestDataHelpers::CompressManifestMemory(DataToEncrypt);
+		CipherHeader.DataSizeCompressed = DataToEncrypt.Num();
+
+		// Encrypt the compressed data.
+		TUniquePtr<ICrypto> Crypto{ FCryptoFactory::Create() };
+		TArray<uint8> AuthTag;
+		CipherHeader.InitializationVector.SetNumUninitialized(AES256_GCM_InitializationVectorSizeInBytes);
+		Crypto->CreateRandomBytes(CipherHeader.InitializationVector);
+		bool bEncryptResult = false;
+		TArray<uint8> Ciphertext = Crypto->Encrypt_AES_256_GCM(DataToEncrypt, InEncryptionSecretKey, CipherHeader.InitializationVector, AuthTag, bEncryptResult);
+		EnumAddFlags(CipherHeader.DataStoredAs, EManifestStorageFlags::Encrypted);
+
+		// Create the actual encryption data block, we don't save the SecretId or AuthTag here since that will get saved in the manifest header.
+		EncryptedData.Data.Reset();
+		FMemoryWriter FinalAr(EncryptedData.Data);
+		FinalAr << CipherHeader;
+		FinalAr << Ciphertext;
+
+		if (bEncryptResult && !FinalAr.IsError())
+		{
+			ObfuscateSensitiveFields();
+			bIsManifestEncrypted = true;
+			FMemory::Memcpy(EncryptionAuthTag.AuthTag, AuthTag.GetData(), BuildPatchServices::AES256_GCM_AuthTagSizeInBytes);
+			EncryptionSecretId = InEncryptionSecretId;
+		}
+		else
+		{
+			UE_LOG(LogBuildPatchManifest, Error, TEXT("FBuildPatchAppManifest::EncryptData: encryption cipher failed. %s"), *InEncryptionSecretId.ToString());
+			bIsManifestEncrypted = false;
+			EncryptedData = FEncryptedData();
+			EncryptionAuthTag.Reset();
+			EncryptionSecretId.Invalidate();
+		}
+
+		// Restore internal lookups either way
+		InitLookups();
+	}
+	else
+	{
+		UE_LOG(LogBuildPatchManifest, Error, TEXT("FBuildPatchAppManifest::EncryptData: SerialiseSensitiveFields failed. %s"), *InEncryptionSecretId.ToString());
+	}
+	return bIsManifestEncrypted;
+}
+
+bool FBuildPatchAppManifest::DecryptData(const TMap<FGuid, TArray<uint8>>& AvailableEncryptionSecrets)
+{
+	if (EncryptionSecretId.IsValid() && bIsManifestEncrypted)
+	{
+		const TArray<uint8>* EncryptionSecretKey = AvailableEncryptionSecrets.Find(EncryptionSecretId);
+		if (EncryptionSecretKey != nullptr)
+		{
+			FCipherHeader CipherHeader;
+			TArray<uint8> Ciphertext;
+			// Read the required values from the encryption block.
+			FMemoryReader EncryptedReader(EncryptedData.Data);
+			EncryptedReader << CipherHeader;
+			EncryptedReader << Ciphertext;
+
+			// Attempt to decrypt the data
+			TUniquePtr<ICrypto> Crypto{ FCryptoFactory::Create() };
+			bool bDecryptResult = false;
+			TArray<uint8> DecryptedData = Crypto->Decrypt_AES_256_GCM(Ciphertext, *EncryptionSecretKey, CipherHeader.InitializationVector, EncryptionAuthTag.AuthTag, bDecryptResult);
+			if (bDecryptResult)
+			{
+				// Uncompress if necessary
+				if (EnumHasAllFlags(CipherHeader.DataStoredAs, EManifestStorageFlags::Compressed))
+				{
+					ManifestDataHelpers::UncompressManifestMemory(DecryptedData, CipherHeader.DataSizeUncompressed);
+				}
+
+				FMemoryReader RawAr(DecryptedData);
+				if (SerialiseSensitiveFields(RawAr))
+				{
+					// Remove the encryption data, leaving EncryptionSecretId intact for info.
+					bIsManifestEncrypted = false;
+					EncryptedData = FEncryptedData();
+					EncryptionAuthTag.Reset();
+
+					// Restore internal lookups.
+					InitLookups();
+
+					return true;
+				}
+				else
+				{
+					UE_LOG(LogBuildPatchManifest, Error, TEXT("FBuildPatchAppManifest::DecryptData: SerialiseSensitiveFields failed. %s"), *EncryptionSecretId.ToString());
+				}
+			}
+			else
+			{
+				UE_LOG(LogBuildPatchManifest, Error, TEXT("FBuildPatchAppManifest::DecryptData: decryption cipher failed. %s"), *EncryptionSecretId.ToString());
+			}
+		}
+		else
+		{
+			UE_LOG(LogBuildPatchManifest, Error, TEXT("FBuildPatchAppManifest::DecryptData: encryption key missing. %s"), *EncryptionSecretId.ToString());
+		}
+	}
+	return false;
+}
+
 void FBuildPatchAppManifest::GetRemovableFiles(const IBuildManifestRef& InOldManifest, TArray<FString>& RemovableFiles) const
 {
 	// Cast manifest parameters
@@ -1347,6 +1761,7 @@ void FBuildPatchAppManifest::GetOutdatedFiles(const FBuildPatchAppManifest* OldM
 	{
 		// Enumerate files in the this file list, that do not exist, or have different hashes in the OldManifest
 		// to be files no longer required by the build
+		TMap<FString, int64> OutdatedFileSizes = bCheckExistingFile ? GetFileSizes(FilesToCheck, InstallDirectory) : TMap<FString, int64>{};
 		for (const FString& FileToCheck : FilesToCheck)
 		{
 			const FFileManifest* NewFile = GetFileManifest(FileToCheck);
@@ -1360,7 +1775,8 @@ void FBuildPatchAppManifest::GetOutdatedFiles(const FBuildPatchAppManifest* OldM
 				// Double check an unchanged file is not missing (size will be -1) or is incorrect size
 				else if (bCheckExistingFile)
 				{
-					const int64 ExistingFileSize = IFileManager::Get().FileSize(*(InstallDirectory / NewFile->Filename));
+					const int64* PrecalculatedSize = OutdatedFileSizes.Find(NewFile->Filename);
+					const int64 ExistingFileSize = PrecalculatedSize ? *PrecalculatedSize : IFileManager::Get().FileSize(*(InstallDirectory / NewFile->Filename));
 					if ((ExistingFileSize < 0) || (ExistingFileSize != NewFile->FileSize))
 					{
 						OutDatedFiles.Add(NewFile->Filename);
@@ -1371,6 +1787,27 @@ void FBuildPatchAppManifest::GetOutdatedFiles(const FBuildPatchAppManifest* OldM
 	}
 }
 
+void FBuildPatchAppManifest::GetMissedFiles(const FString& InstallDirectory, const TSet<FString>& FilesToCheck, TSet<FString>& OutMissedFiles) const
+{
+	ensure(!InstallDirectory.IsEmpty());
+	if (InstallDirectory.IsEmpty())
+	{
+		return;
+	}
+	TSet<FString> Filenames;
+	GetFileList(Filenames);
+	// Enumerate files in the this file list, that do not exist on disk or missed in manifest
+	const TMap<FString, int64> OutdatedFileSizes = GetFileSizes(FilesToCheck.Intersect(Filenames), InstallDirectory);
+	
+	for (const TPair<FString, int64>& Item : OutdatedFileSizes)
+	{
+		if (Item.Value < 0LL) // if file doesn't exist
+		{
+			OutMissedFiles.Add(Item.Key); // add to list of missing files.
+		}
+	}
+}
+
 bool FBuildPatchAppManifest::IsFileOutdated(const FBuildPatchAppManifestRef& OldManifest, const FString& Filename) const
 {
 	return IsFileOutdated(OldManifest.Get(), Filename);
@@ -1392,7 +1829,7 @@ bool FBuildPatchAppManifest::IsFileOutdated(const FBuildPatchAppManifest& OldMan
 		return true;
 	}
 	// Different hash means different file
-	if (OldFile->FileHash != NewFile->FileHash)
+	if (OldFile->SHA1Hash != NewFile->SHA1Hash)
 	{
 		return true;
 	}
diff --git a/Engine/Source/Runtime/Online/BuildPatchServices/Private/BuildPatchManifest.h b/Engine/Source/Runtime/Online/BuildPatchServices/Private/BuildPatchManifest.h
--- a/Engine/Source/Runtime/Online/BuildPatchServices/Private/BuildPatchManifest.h
+++ b/Engine/Source/Runtime/Online/BuildPatchServices/Private/BuildPatchManifest.h
@@ -52,6 +52,7 @@ namespace BuildPatchServices
 	class FBuildPatchInstaller;
 	class FManifestBuilder;
 	class FManifestData;
+	class FDirectoryChunker;
 	class FChunkDeltaOptimiser;
 	class FBuildPatchManifestSet;
 }
@@ -71,6 +72,7 @@ class FBuildPatchAppManifest
 	friend class FBuildDiffManifests;
 	friend class FManifestUObject;
 	friend class BuildPatchServices::FManifestData;
+	friend class BuildPatchServices::FDirectoryChunker;
 	friend class BuildPatchServices::FChunkDeltaOptimiser;
 	friend class BuildPatchServices::FBuildPatchManifestSet;
 public:
@@ -96,21 +98,32 @@ public:
 	~FBuildPatchAppManifest();
 
 	// START IBuildManifest Interface
-	virtual uint32  GetAppID() const override;
+	virtual BuildPatchServices::EFeatureLevel GetFeatureLevel() const override;
+	virtual uint32 GetAppID() const override;
 	virtual const FString& GetAppName() const override;
 	virtual const FString& GetVersionString() const override;
+	virtual const FString& GetUniqueBuildId() const override;
+	virtual const FGuid& GetEncryptionSecretId(bool* bOutIsManifestEncrypted = nullptr) const override;
+	virtual bool IsManifestEncrypted() const override;
+	virtual TSet<FGuid> GetNecessaryEncryptionSecretIds() const override;
 	virtual const FString& GetLaunchExe() const override;
 	virtual const FString& GetLaunchCommand() const override;
 	virtual const TSet<FString>& GetPrereqIds() const override;
 	virtual const FString& GetPrereqName() const override;
 	virtual const FString& GetPrereqPath() const override;
 	virtual const FString& GetPrereqArgs() const override;
+	virtual const FString& GetUninstallActionPath() const override;
+	virtual const FString& GetUninstallActionArgs() const override;
 	virtual int64 GetDownloadSize() const override;
 	virtual int64 GetDownloadSize(const TSet<FString>& Tags) const override;
+	virtual int64 GetDownloadSize(const TSet<const BuildPatchServices::FFileManifest*>& TaggedFiles) const;
 	virtual int64 GetDeltaDownloadSize(const TSet<FString>& Tags, const IBuildManifestRef& PreviousVersion) const override;
 	virtual int64 GetDeltaDownloadSize(const TSet<FString>& Tags, const IBuildManifestRef& PreviousVersion, const TSet<FString>& PreviousTags) const override;
+	virtual int64 GetDeltaDownloadSize(const TSet<const BuildPatchServices::FFileManifest*>& InTaggedManifests, const IBuildManifestRef& InPreviousVersion, const TSet<FString>& InFilesInstalled, const TSet<FGuid>& InChunksInstalled) const;
+	virtual int64 GetDeltaDownloadSize(const TSet<FString>& AllTaggedFiles, const IBuildManifestRef& InPreviousVersion, const TSet<FString>& InFilesInstalled, const TSet<FGuid>& InChunksInstalled) const;
 	virtual int64 GetBuildSize() const override;
 	virtual int64 GetBuildSize(const TSet<FString>& Tags) const override;
+	virtual int64 GetBuildSize(const TSet<const BuildPatchServices::FFileManifest*>& TaggedFiles) const;
 	virtual TArray<FString> GetBuildFileList() const override;
 	virtual TArray<FStringView> GetBuildFileListView() const override;
 	virtual TArray<FString> GetBuildFileList(const TSet<FString>& Tags) const override;
@@ -118,14 +131,23 @@ public:
 	virtual int64 GetFileSize(FStringView Filename) const override;
 	virtual int64 GetFileSize(const TArray<FString>& Filenames) const override;
 	virtual int64 GetFileSize(const TSet  <FString>& Filenames) const override;
-	virtual bool GetFileHash(const FString& Filename, FSHAHash& OutHash) const override;
+	virtual EFileMetaFlags GetFileMetaFlags(const FString& Filename) const override;
+	virtual const FMD5Hash* GetFileMD5Hash(const FString& Filename) const override;
+	virtual const FSHAHash* GetFileSHA1Hash(const FString& Filename) const override;
+	virtual const FSHA256Signature* GetFileSHA256Hash(const FString& Filename) const override;
+	virtual const FString* GetFileMIMEType(const FString& Filename) const override;
 	virtual TSet<FString> GetFileTagList() const override;
 	virtual void GetFileTagList(TSet<FString>& Tags) const override;
 	virtual void GetOutdatedFiles(const IBuildManifestRef& OldManifest, TSet<FString>& OutdatedFiles) const override;
 	virtual void GetRemovableFiles(const IBuildManifestRef& OldManifest, TArray<FString>& RemovableFiles) const override;
 	virtual void GetRemovableFiles(const TCHAR* InstallPath, TArray<FString>& RemovableFiles) const override;
+	virtual void GetChunksRequiredForFiles(const TSet<FString>& Filenames, TSet<FGuid>& RequiredChunks) const override;
+	virtual int64 GetDataSize(const FGuid& DataGuid) const override;
+	virtual bool GetChunkShaHash(const FGuid& ChunkGuid, FSHAHash& OutHash) const override;
+	virtual bool GetChunkEncryptionSecretId(const FGuid& ChunkGuid, FGuid& OutEncryptionSecretId) const override;
 	virtual bool NeedsResaving() const override;
 	virtual void CopyCustomFields(const IBuildManifestRef& Other, bool bClobber) override;
+	virtual TArray<FString> GetCustomFieldNames() const override;
 	virtual const IManifestFieldPtr GetCustomField(const FString& FieldName) const override;
 	virtual const IManifestFieldPtr SetCustomField(const FString& FieldName, const FString& Value) override;
 	virtual const IManifestFieldPtr SetCustomField(const FString& FieldName, const double& Value) override;
@@ -162,12 +184,14 @@ public:
 
 	/**
 	 * Saves out the manifest information.
-	 * @param Filename      The file to save to.
-	 * @param SaveFormat    The feature level that the intended client has support for, which the manifest will need saving as.
-	 *                      A manifest file cannot be downgraded, the function will fail if the provided value is less than GetFeatureLevel().
+	 * @param Filename      IN   The file to save to.
+	 * @param SaveFormat    IN   The feature level that the intended client has support for, which the manifest will need saving as.
+	 *                           A manifest file cannot be downgraded, the function will fail if the provided value is less than GetFeatureLevel().
+	 * @param OutSHA1Hash   OUT  An output variable which will be populated with the SHA1 hash of the written manifest.
+	 * @param OutMD5Hash    OUT  An output variable which will be populated with the MD5 hash of the written manifest.
 	 * @return True if successful.
 	 */
-	virtual bool SaveToFile(const FString& Filename, BuildPatchServices::EFeatureLevel SaveFormat = BuildPatchServices::EFeatureLevel::Latest);
+	virtual bool SaveToFile(const FString& Filename, BuildPatchServices::EFeatureLevel SaveFormat, FSHAHash* OutSHA1Hash = nullptr, FMD5Hash* OutMD5Hash = nullptr);
 
 	/**
 	 * Creates the object in JSON format
@@ -175,19 +199,6 @@ public:
 	 */
 	virtual void SerializeToJSON(FString& JSONOutput);
 
-	/**
-	 * Gets the feature level for this manifest.
-	 * @return		The highest available feature support
-	 */
-	virtual BuildPatchServices::EFeatureLevel GetFeatureLevel() const;
-
-	/**
-	 * Provides the set of chunks required to produce the given files.
-	 * @param Filenames         IN      The set of files.
-	 * @param RequiredChunks    OUT     The set of chunk GUIDs needed for those files.
-	 */
-	virtual void GetChunksRequiredForFiles(const TSet<FString>& Filenames, TSet<FGuid>& RequiredChunks) const;
-
 	/**
 	 * Get the number of times a chunks is referenced in this manifest
 	 * @param ChunkGuid		The chunk GUID
@@ -195,13 +206,6 @@ public:
 	 */
 	virtual uint32 GetNumberOfChunkReferences(const FGuid& ChunkGuid) const;
 
-	/**
-	 * Returns the size of a particular data file by it's GUID.
-	 * @param DataGuid		The GUID for the data
-	 * @return		File size, or 0 if this data is not in the manifest.
-	 */
-	virtual int64 GetDataSize(const FGuid& DataGuid) const;
-
 	/**
 	 * Returns the total size of all data files in it's list.
 	 * @param DataGuids		The GUID array for the data
@@ -229,9 +233,16 @@ public:
 	 * @param Tags					The tags for the required file groups.
 	 * @param TaggedFiles	OUT		Receives the tagged files.
 	 */
-	virtual void GetTaggedFileList(const TSet<FString>& Tags, TArray<FString>&     TaggedFiles) const;
+	virtual void GetTaggedFileList(const TSet<FString>& Tags, TArray<FString>& TaggedFiles) const;
 	virtual void GetTaggedFileList(const TSet<FString>& Tags, TArray<FStringView>& TaggedFiles) const;
-	virtual void GetTaggedFileList(const TSet<FString>& Tags, TSet<FString>&       TaggedFiles) const;
+	virtual void GetTaggedFileList(const TSet<FString>& Tags, TSet<FString>& TaggedFiles) const;
+
+	/**
+	 * Get the set of manifests that are tagged with the provided tags
+	 * @param Tags					The tags for the required file groups.
+	 * @return The set of tagged file manifests.
+	 */
+	virtual TSet<const BuildPatchServices::FFileManifest*> GetTaggedFileManifests(const TSet<FString>& Tags) const;
 
 	/**
 	 * Get the list of Guids for all chunks referenced by this manifest
@@ -261,14 +272,6 @@ public:
 	 */
 	virtual bool GetChunkHash(const FGuid& ChunkGuid, uint64& OutHash) const;
 
-	/**
-	 * Gets the SHA1 hash for a given chunk
-	 * @param ChunkGuid		IN		The guid of the chunk to get hash for
-	 * @param OutHash		OUT		Receives the hash value if found
-	 * @return	true if we had the hash for this chunk
-	 */
-	virtual bool GetChunkShaHash(const FGuid& ChunkGuid, FSHAHash& OutHash) const;
-
 	/**
 	 * Gets the FChunkInfo for a given chunk.
 	 * @param ChunkGuid     The guid of the chunk to get hash for.
@@ -284,6 +287,14 @@ public:
 	 */
 	virtual bool GetFileHash(const FGuid& FileGuid, FSHAHash& OutHash) const; // DEPRECATE ME
 
+	/**
+	 * Gets the file hash for a given file
+	 * @param Filename		IN		The filename in the build
+	 * @param OutHash		OUT		Receives the hash value if found
+	 * @return	true if we had the hash for this file
+	 */
+	virtual bool GetFileHash(const FString& Filename, FSHAHash& OutHash) const;
+
 	/**
 	 * Gets the file hash for given file data. Valid for non-chunked manifest
 	 * @param FileGuid		IN		The guid of the file data to get hash for
@@ -301,6 +312,7 @@ public:
 	 * @return the number of chunks added to the ChunksAvailable set.
 	 */
 	virtual int32 EnumerateProducibleChunks(const TSet<FString>& TagSet, const TSet<FGuid>& ChunksRequired, TSet<FGuid>& ChunksAvailable) const;
+	virtual int32 EnumerateProducibleChunks(const TSet<const BuildPatchServices::FFileManifest*>& TaggedFiles, const TSet<FGuid>& ChunksRequired, TSet<FGuid>& ChunksAvailable) const;
 
 	/**
 	 * Gets a list of files that have changed or are new in the this manifest, compared to those in the old manifest, or are missing from disk.
@@ -312,6 +324,14 @@ public:
 	virtual void GetOutdatedFiles(const FBuildPatchAppManifest*    OldManifest, const FString& InstallDirectory, TSet<FString>& OutDatedFiles) const;
 	virtual void GetOutdatedFiles(const FBuildPatchAppManifest*    OldManifest, const FString& InstallDirectory, const TSet<FString>& FilesToCheck, TSet<FString>& OutDatedFiles) const;
 
+	/**
+	 * Gets a list of files mentioned in the manifest that are missing from disk.
+	 * @param InstallDirectory [in]  The build installation directory to check for missing files.
+	 * @param FilesToCheck     [in]  The list of files to check.
+	 * @param OutMissedFiles   [out] The files that are missing from disk.
+	 */
+	virtual void GetMissedFiles(const FString& InstallDirectory, const TSet<FString>& FilesToCheck, TSet<FString>& OutMissedFiles) const;
+
 	/**
 	 * Check a single file to see if it will be effected by patching from a previous version.
 	 * @param OldManifest		The Build Manifest that is currently installed. Shared Ref - Implicitly valid.
@@ -331,12 +351,19 @@ public:
 	/** @return True if any files in this manifest have file attributes to be set */
 	virtual bool HasFileAttributes() const;
 
+	bool EncryptData(const FGuid& InEncryptionSecretId, const TArray<uint8>& InEncryptionSecretKey);
+
+	bool DecryptData(const TMap<FGuid, TArray<uint8>>& AvailableEncryptionSecrets);
+
 private:
 	/**
 	 * Destroys any memory we have allocated and clears out ready for generation of a new manifest
 	 */
 	void DestroyData();
 
+	bool SerialiseSensitiveFields(FArchive& Ar);
+	void ObfuscateSensitiveFields();
+
 protected:
 	/**
 	 * Setups the lookup maps that optimize data access, should be called when Data changes
@@ -356,8 +383,16 @@ protected:
 	BuildPatchServices::FChunkDataList ChunkDataList;
 	BuildPatchServices::FFileManifestList FileManifestList;
 	BuildPatchServices::FCustomFields CustomFields;
+	BuildPatchServices::FEncryptedData EncryptedData;
 
 private:
+	// The encryption secret ID and hash used for this manifest, if encrypted.
+	// These are serialised with the manifest header, not the main data.
+	FGuid EncryptionSecretId;
+	BuildPatchServices::FAESAuthTag EncryptionAuthTag;
+
+	// Whether the manifest is currently encrypted.
+	bool bIsManifestEncrypted = false;
 
 	/** Holds the handle to our PreExit delegate */
 	FDelegateHandle OnPreExitHandle;
diff --git a/Engine/Source/Runtime/Online/BuildPatchServices/Private/BuildPatchManifestSet.cpp b/Engine/Source/Runtime/Online/BuildPatchServices/Private/BuildPatchManifestSet.cpp
--- a/Engine/Source/Runtime/Online/BuildPatchServices/Private/BuildPatchManifestSet.cpp
+++ b/Engine/Source/Runtime/Online/BuildPatchServices/Private/BuildPatchManifestSet.cpp
@@ -143,9 +143,45 @@ namespace BuildPatchServices
 			}
 		}
 
+		virtual void GetReferencedChunks(TArray<FGuid>& DataGuids) const override
+		{
+			// For an array instead of set, we would expect to provide duplicates
+			for (const FBuildPatchInstallerAction& InstallerAction : InstallerActions)
+			{
+				if (!InstallerAction.IsUninstall())
+				{
+					for (const FString& TaggedFile : InstallerAction.GetTaggedFiles())
+					{
+						FFileReference const* const LookupResult = NewFileLookup.Find(TaggedFile);
+						if (LookupResult)
+						{
+							FFileManifest const* const TaggedFileManifest = LookupResult->Get<0>();
+							for (const FChunkPart& ChunkPart : TaggedFileManifest->ChunkParts)
+							{
+								DataGuids.Add(ChunkPart.Guid);
+							}
+						}
+					}
+				}
+			}
+		}
+
+		FDataReference const* const DataLookupFind(const FGuid& DataGuid) const
+		{
+			static thread_local FDataReference const* LatestLookupResult = nullptr;
+			static thread_local FGuid LatestGuid = FGuid();
+			if (LatestGuid == DataGuid && LatestLookupResult != nullptr)
+			{
+				return LatestLookupResult;
+			}
+			LatestGuid = DataGuid;
+			LatestLookupResult = DataLookup.Find(DataGuid);
+			return LatestLookupResult;
+		}
+
 		virtual uint64 GetDownloadSize(const FGuid& DataGuid) const override
 		{
-			FDataReference const * const LookupResult = DataLookup.Find(DataGuid);
+			FDataReference const* const LookupResult = DataLookupFind(DataGuid);
 			if (LookupResult)
 			{
 				return LookupResult->Get<0>()->FileSize;
@@ -161,7 +197,7 @@ namespace BuildPatchServices
 		virtual bool GetChunkShaHash(const FGuid& DataGuid, FSHAHash& OutHash) const override
 		{
 			static const uint8 Zero[FSHA1::DigestSize] = { 0 };
-			FDataReference const * const LookupResult = DataLookup.Find(DataGuid);
+			FDataReference const* const LookupResult = DataLookupFind(DataGuid);
 			if (LookupResult)
 			{
 				OutHash = LookupResult->Get<0>()->ShaHash;
@@ -172,7 +208,7 @@ namespace BuildPatchServices
 
 		virtual FString GetDataFilename(const FGuid& DataGuid) const override
 		{
-			FDataReference const * const LookupResult = DataLookup.Find(DataGuid);
+			FDataReference const* const LookupResult = DataLookupFind(DataGuid);
 			if (LookupResult)
 			{
 				return LookupResult->Get<2>()->GetCloudSubdirectory() / FBuildPatchUtils::GetDataFilename(*LookupResult->Get<1>(), DataGuid);
@@ -180,6 +216,56 @@ namespace BuildPatchServices
 			return FString();
 		}
 
+		virtual FGuid GetDataSecretId(const FGuid& DataGuid) const override
+		{
+			FDataReference const* const LookupResult = DataLookupFind(DataGuid);
+			if (LookupResult)
+			{
+				return LookupResult->Get<0>()->EncryptionSecretId;
+			}
+			return FGuid();
+		}
+
+		virtual FString GetChunkCloudSubdirectory(const FGuid& DataGuid) const override
+		{
+			FDataReference const* const LookupResult = DataLookupFind(DataGuid);
+			if (LookupResult)
+			{
+				return LookupResult->Get<2>()->GetCloudSubdirectory();
+			}
+			return FString();
+		}
+
+		virtual FString GetChunkFilename(const FGuid& DataGuid) const override
+		{
+			FDataReference const* const LookupResult = DataLookupFind(DataGuid);
+			if (LookupResult)
+			{
+				return FBuildPatchUtils::GetDataFilename(*LookupResult->Get<1>(), DataGuid);
+			}
+			return FString();
+		}
+
+		virtual FString GetChunkUniqueBuildId(const FGuid& DataGuid) const override
+		{
+			FDataReference const* const LookupResult = DataLookupFind(DataGuid);
+			if (LookupResult)
+			{
+				return LookupResult->Get<2>()->GetUniqueBuildId();
+			}
+			return FString();
+		}
+
+		virtual bool GetChunkFromDeltaStatus(const FGuid& DataGuid) const override
+		{
+			FDataReference const* const LookupResult = DataLookupFind(DataGuid);
+			if (LookupResult)
+			{
+				return LookupResult->Get<0>()->bChunkFromDelta;
+			}
+			return false;
+		}
+
 		virtual int32 GetNumExpectedFiles() const
 		{
 			return NewFileLookup.Num();
@@ -208,6 +294,35 @@ namespace BuildPatchServices
 			}
 		}
 
+		virtual void GetMissedFiles(const FString& InstallDirectory, TSet<FString>& OutMissedFiles) const override
+		{
+			TSet<FString> TempFileSet;
+			for (const FBuildPatchInstallerAction& InstallerAction : InstallerActions)
+			{
+				// It expects to use method for Update.
+				if (!InstallerAction.IsUninstall() && !InstallerAction.IsInstall())
+				{
+					TempFileSet.Reset();
+					const FBuildPatchAppManifest* CurrentManifest = InstallerAction.TryGetCurrentManifest();
+					if (CurrentManifest)
+					{
+						// There is no reliable way to determine the list of tags with which the binary was previously installed.
+						// At the same time, InstallerAction.GetTaggedFiles() cannot be used because it contains files from the new manifest.
+						TSet<FString> AllFilenames;
+						CurrentManifest->GetFileList(AllFilenames);
+						CurrentManifest->GetMissedFiles(InstallDirectory / InstallerAction.GetInstallSubdirectory(), AllFilenames, TempFileSet);
+					}
+					else
+					{
+						// Treat all files as missing if the current manifest is unavailable.
+						TempFileSet.Append(InstallerAction.GetTaggedFiles());
+					}
+					
+					Algo::Transform(TempFileSet, OutMissedFiles, [&](const FString& Filename) { return InstallerAction.GetInstallSubdirectory() / Filename; });
+				}
+			}
+		}
+
 		virtual void GetRemovableFiles(TSet<FString>& FilesToRemove) const override
 		{
 			// We enumerate all trackable files, adding anything not in the NewFileLookup map.
@@ -230,7 +345,7 @@ namespace BuildPatchServices
 
 		virtual const BuildPatchServices::FFileManifest* GetCurrentFileManifest(const FString& BuildFile) const override
 		{
-			FFileReference const * const LookupResult = CurrentFileLookup.Find(BuildFile);
+			FFileReference const* const LookupResult = CurrentFileLookup.Find(BuildFile);
 			if (LookupResult)
 			{
 				return LookupResult->Get<0>();
@@ -240,7 +355,7 @@ namespace BuildPatchServices
 
 		virtual const BuildPatchServices::FFileManifest* GetNewFileManifest(const FString& BuildFile) const
 		{
-			FFileReference const * const LookupResult = NewFileLookup.Find(BuildFile);
+			FFileReference const* const LookupResult = NewFileLookup.Find(BuildFile);
 			if (LookupResult)
 			{
 				return LookupResult->Get<0>();
@@ -300,7 +415,7 @@ namespace BuildPatchServices
 
 		virtual bool IsFileRepairAction(const FString& Filename) const override
 		{
-			FFileReference const * const LookupResult = NewFileLookup.Find(Filename);
+			FFileReference const* const LookupResult = NewFileLookup.Find(Filename);
 			if (LookupResult)
 			{
 				return LookupResult->Get<2>()->IsRepair();
@@ -320,7 +435,7 @@ namespace BuildPatchServices
 			return false;
 		}
 
-		virtual bool IsRepairOnly() const override
+		virtual bool IsExclusivelyRepairing() const override
 		{
 			bool bHasRepair = false;
 			for (const FBuildPatchInstallerAction& InstallerAction : InstallerActions)
@@ -352,6 +467,41 @@ namespace BuildPatchServices
 			return false;
 		}
 
+		virtual bool IsExclusivelyVerifying() const override
+		{
+			bool bHasVerifyOnly = false;
+			for (const FBuildPatchInstallerAction& InstallerAction : InstallerActions)
+			{
+				if (InstallerAction.IsVerifyOnly())
+				{
+					bHasVerifyOnly = true;
+				}
+				else
+				{
+					return false;
+				}
+			}
+			return bHasVerifyOnly;
+		}
+
+		virtual void GetUninstallActionInfo(TArray<FUninstallActionInfo>& UninstallActionsInfo) const override
+		{
+			for (const FBuildPatchInstallerAction& InstallerAction : InstallerActions)
+			{
+				if (InstallerAction.IsUninstall())
+				{
+					const FBuildPatchAppManifest& Manifest = InstallerAction.GetInstallOrCurrentManifest();
+					const FString& UninstallActionPath = Manifest.GetUninstallActionPath();
+					if (!UninstallActionPath.IsEmpty())
+					{
+						FUninstallActionInfo& UninstallActionInfo = UninstallActionsInfo.AddDefaulted_GetRef();
+						UninstallActionInfo.Path = UninstallActionPath;
+						UninstallActionInfo.Args = Manifest.GetUninstallActionArgs();
+					}
+				}
+			}
+		}
+
 	private:
 		TArray<FBuildPatchInstallerAction> InstallerActions;
 	};
diff --git a/Engine/Source/Runtime/Online/BuildPatchServices/Private/BuildPatchMergeManifests.cpp b/Engine/Source/Runtime/Online/BuildPatchServices/Private/BuildPatchMergeManifests.cpp
--- a/Engine/Source/Runtime/Online/BuildPatchServices/Private/BuildPatchMergeManifests.cpp
+++ b/Engine/Source/Runtime/Online/BuildPatchServices/Private/BuildPatchMergeManifests.cpp
@@ -52,7 +52,7 @@ namespace MergeHelpers
 		return bSuccess;
 	}
 
-	bool ReinitialiseChunkInfoList(const TArray<BuildPatchServices::FFileManifest>& FileManifestList, const FBuildPatchAppManifestPtr& ManifestA, const FBuildPatchAppManifestPtr& ManifestB, TArray<BuildPatchServices::FChunkInfo>& ChunkList)
+	bool ReinitialiseChunkInfoList(const TArray<BuildPatchServices::FFileManifest>& FileManifestList, const FBuildPatchAppManifest& ManifestA, const FBuildPatchAppManifest& ManifestB, TArray<BuildPatchServices::FChunkInfo>& ChunkList)
 	{
 		using namespace BuildPatchServices;
 		ChunkList.Reset();
@@ -66,10 +66,10 @@ namespace MergeHelpers
 				if (bAlreadyInSet == false)
 				{
 					// Find the chunk info
-					const FChunkInfo* ChunkInfo = ManifestB->GetChunkInfo(FileChunkPart.Guid);
-					if (ChunkInfo == nullptr && ManifestA.IsValid())
+					const FChunkInfo* ChunkInfo = ManifestB.GetChunkInfo(FileChunkPart.Guid);
+					if (ChunkInfo == nullptr)
 					{
-						ChunkInfo = ManifestA->GetChunkInfo(FileChunkPart.Guid);
+						ChunkInfo = ManifestA.GetChunkInfo(FileChunkPart.Guid);
 					}
 					if (ChunkInfo == nullptr)
 					{
@@ -78,7 +78,9 @@ namespace MergeHelpers
 					}
 					else
 					{
-						ChunkList.Add(*ChunkInfo);
+						FChunkInfo ChunkTemp = *ChunkInfo;
+						ChunkTemp.bChunkFromDelta = FileChunkPart.bChunkFromDelta;
+						ChunkList.Add(MoveTemp(ChunkTemp));
 					}
 				}
 			}
@@ -87,10 +89,10 @@ namespace MergeHelpers
 	}
 }
 
-bool FBuildMergeManifests::MergeManifests(const FString& ManifestFilePathA, const FString& ManifestFilePathB, const FString& ManifestFilePathC, const FString& NewVersionString, const FString& SelectionDetailFilePath)
+bool FBuildMergeManifests::MergeManifests(const FString& ManifestFilePathA, const FString& ManifestFilePathB, const FString& ManifestFilePathC, const FString& NewVersionString, const FString& SelectionDetailFilePath, const TMap<FGuid, TArray<uint8>>& AvailableEncryptionSecrets)
 {
 	using namespace BuildPatchServices;
-	bool bSuccess = true;
+
 	FCriticalSection UObjectAllocationLock;
 
 	TFunction<FBuildPatchAppManifestPtr()> TaskManifestA = [&UObjectAllocationLock, &ManifestFilePathA]()
@@ -101,13 +103,13 @@ bool FBuildMergeManifests::MergeManifests(const FString& ManifestFilePathA, cons
 	{
 		return MergeHelpers::LoadManifestFile(ManifestFilePathB, &UObjectAllocationLock);
 	};
-	typedef TPair<TSet<FString>,TSet<FString>> FStringSetPair;
+	typedef TPair<TSet<FString>, TSet<FString>> FStringSetPair;
 	FThreadSafeBool bSelectionDetailSuccess = false;
 	TFunction<FStringSetPair()> TaskSelectionInfo = [&SelectionDetailFilePath, &bSelectionDetailSuccess]()
 	{
 		bSelectionDetailSuccess = true;
 		FStringSetPair StringSetPair;
-		if(SelectionDetailFilePath.IsEmpty() == false)
+		if (SelectionDetailFilePath.IsEmpty() == false)
 		{
 			FString SelectionDetailFileData;
 			bSelectionDetailSuccess = FFileHelper::LoadFileToString(SelectionDetailFileData, *SelectionDetailFilePath);
@@ -168,6 +170,23 @@ bool FBuildMergeManifests::MergeManifests(const FString& ManifestFilePathA, cons
 		return false;
 	}
 
+	// We cannot use merge manifests unless we have all required secret keys for both, otherwise we cannot match file names, and encrypt the produced manifest file.
+	TSet<FGuid> AvailableSecrets;
+	AvailableEncryptionSecrets.GetKeys(AvailableSecrets);
+	TSet<FGuid> NecessaryEncryptionSecrets;
+	NecessaryEncryptionSecrets.Append(ManifestA->GetNecessaryEncryptionSecretIds());
+	NecessaryEncryptionSecrets.Append(ManifestB->GetNecessaryEncryptionSecretIds());
+	TSet<FGuid> MissingSecrets = NecessaryEncryptionSecrets.Difference(AvailableSecrets);
+	if (MissingSecrets.Num() > 0)
+	{
+		UE_LOG(LogMergeManifests, Error, TEXT("Cannot continue without all necessary secret keys."));
+		for (const FGuid& SecretId : MissingSecrets)
+		{
+			UE_LOG(LogMergeManifests, Error, TEXT("    Missing secret with ID: %s"), *SecretId.ToString());
+		}
+		return false;
+	}
+
 	// Check if the selection detail had an error
 	if (bSelectionDetailSuccess == false)
 	{
@@ -187,62 +206,120 @@ bool FBuildMergeManifests::MergeManifests(const FString& ManifestFilePathA, cons
 		SelectionDetail.Key = SelectionDetail.Key.Difference(SelectionDetail.Value);
 	}
 
+	IBuildManifestPtr Result = MergeManifests(ManifestA.ToSharedRef(), ManifestB.ToSharedRef(), NewVersionString, SelectionDetail.Key, SelectionDetail.Value);
+
+	// Save the new manifest out if we didn't register a failure
+	if (Result.IsValid())
+	{
+		FBuildPatchAppManifestRef ResultRef = StaticCastSharedRef<FBuildPatchAppManifest>(Result.ToSharedRef());
+		// If ManifestB was encrypted, encrypt ManifestC if we had the key given.
+		bool bEncryptionSuccess = true;
+		const FGuid EncryptionSecretId = ManifestB->GetEncryptionSecretId();
+		if (EncryptionSecretId.IsValid())
+		{
+			const TArray<uint8>* EncryptionSecretKey = AvailableEncryptionSecrets.Find(EncryptionSecretId);
+			if (EncryptionSecretKey != nullptr)
+			{
+				UE_LOG(LogMergeManifests, Log, TEXT("Encrypting output manifest with secret %s."), *EncryptionSecretId.ToString());
+				bEncryptionSuccess = ResultRef->EncryptData(EncryptionSecretId, *EncryptionSecretKey);
+				if (!bEncryptionSuccess)
+				{
+					UE_LOG(LogMergeManifests, Error, TEXT("Systematic failure to encrypt new manifest"));
+				}
+			}
+			else
+			{
+				UE_LOG(LogMergeManifests, Error, TEXT("Failed to encrypt new manifest, missing encryption key for %s"), *EncryptionSecretId.ToString());
+				bEncryptionSuccess = false;
+			}
+		}
+		if (bEncryptionSuccess)
+		{
+			TUniquePtr<IFileSystem> FileSystem(FFileSystemFactory::Create());
+			const FString TmpManifestFilePathC = ManifestFilePathC + TEXT("tmp");
+			if (!ResultRef->SaveToFile(TmpManifestFilePathC, ResultRef->ManifestMeta.FeatureLevel) || !FileSystem->MoveFile(*ManifestFilePathC, *TmpManifestFilePathC))
+			{
+				UE_LOG(LogMergeManifests, Error, TEXT("Failed to save new manifest %s"), *ManifestFilePathC);
+				Result.Reset();
+			}
+		}
+		else
+		{
+				Result.Reset();
+		}
+	}
+	else
+	{
+		UE_LOG(LogMergeManifests, Error, TEXT("Not saving new manifest due to previous errors."));
+	}
+
+	return Result.IsValid();
+}
+
+IBuildManifestPtr FBuildMergeManifests::MergeManifests(const IBuildManifestRef& InManifestA, const IBuildManifestRef& InManifestB, const FString& NewVersionString, const TSet<FString>& FilesFromA, const TSet<FString>& FilesFromB)
+{
+	using namespace BuildPatchServices;
+	bool bSuccess = true;
+	FBuildPatchAppManifestRef ManifestA = StaticCastSharedRef<FBuildPatchAppManifest>(InManifestA);
+	FBuildPatchAppManifestRef ManifestB = StaticCastSharedRef<FBuildPatchAppManifest>(InManifestB);
+
+	// We can't do this if the manifests are currently encrypted.
+	if (ManifestA->IsManifestEncrypted() || ManifestB->IsManifestEncrypted())
+	{
+		UE_LOG(LogMergeManifests, Error, TEXT("Cannot merge encrypted manifests, make sure the secrets were provided."));
+		return nullptr;
+	}
+
 	// Create the new manifest
-	FBuildPatchAppManifest MergedManifest;
+	FBuildPatchAppManifestRef MergedManifest = MakeShareable(new FBuildPatchAppManifest());
 
 	// Copy basic info from B, preserving the generated build ID
 	{
-		FString NewBuildId = MoveTemp(MergedManifest.ManifestMeta.BuildId);
-		MergedManifest.ManifestMeta = ManifestB->ManifestMeta;
-		MergedManifest.CustomFields = ManifestB->CustomFields;
-		MergedManifest.ManifestMeta.BuildId = MoveTemp(NewBuildId);
+		FString NewBuildId = MoveTemp(MergedManifest->ManifestMeta.BuildId);
+		MergedManifest->ManifestMeta = ManifestB->ManifestMeta;
+		MergedManifest->CustomFields = ManifestB->CustomFields;
+		MergedManifest->ManifestMeta.BuildId = MoveTemp(NewBuildId);
 	}
 
 	// Set the new version string
-	MergedManifest.ManifestMeta.BuildVersion = NewVersionString;
+	MergedManifest->ManifestMeta.BuildVersion = NewVersionString;
 
 	// Copy the file manifests required from A
-	bSuccess = MergeHelpers::CopyFileDataFromManifestToArray(SelectionDetail.Key, ManifestA, MergedManifest.FileManifestList.FileList) && bSuccess;
+	bSuccess = MergeHelpers::CopyFileDataFromManifestToArray(FilesFromA, ManifestA, MergedManifest->FileManifestList.FileList) && bSuccess;
 
 	// Copy the file manifests required from B
-	bSuccess = MergeHelpers::CopyFileDataFromManifestToArray(SelectionDetail.Value, ManifestB, MergedManifest.FileManifestList.FileList) && bSuccess;
+	bSuccess = MergeHelpers::CopyFileDataFromManifestToArray(FilesFromB, ManifestB, MergedManifest->FileManifestList.FileList) && bSuccess;
 
 	// Call OnPostLoad for the file manifest list before entering chunk info.
-	MergedManifest.FileManifestList.OnPostLoad();
+	MergedManifest->FileManifestList.OnPostLoad();
 
 	// Fill out the chunk list in order of reference
-	bSuccess = MergeHelpers::ReinitialiseChunkInfoList(MergedManifest.FileManifestList.FileList, ManifestA, ManifestB, MergedManifest.ChunkDataList.ChunkList) && bSuccess;
+	bSuccess = MergeHelpers::ReinitialiseChunkInfoList(MergedManifest->FileManifestList.FileList, ManifestA.Get(), ManifestB.Get(), MergedManifest->ChunkDataList.ChunkList) && bSuccess;
 
-	// Save the new manifest out if we didn't register a failure
-	if (bSuccess)
-	{
-		TUniquePtr<IFileSystem> FileSystem(FFileSystemFactory::Create());
-		MergedManifest.InitLookups();
-		const FString TmpManifestFilePathC = ManifestFilePathC + TEXT("tmp");
-		if (!MergedManifest.SaveToFile(TmpManifestFilePathC, MergedManifest.ManifestMeta.FeatureLevel) || !FileSystem->MoveFile(*ManifestFilePathC, *TmpManifestFilePathC))
-		{
-			UE_LOG(LogMergeManifests, Error, TEXT("Failed to save new manifest %s"), *ManifestFilePathC);
-			bSuccess = false;
-		}
-	}
-	else
+	// Return invalid if there were any issues.
+	if (!bSuccess)
 	{
-		UE_LOG(LogMergeManifests, Error, TEXT("Not saving new manifest due to previous errors."));
+		return nullptr;
 	}
 
-	return bSuccess;
+	MergedManifest->InitLookups();
+	return MergedManifest;
 }
 
-FBuildPatchAppManifestPtr FBuildMergeManifests::MergeDeltaManifest(const FBuildPatchAppManifestRef& Manifest, const FBuildPatchAppManifestRef& Delta)
+FBuildPatchAppManifestPtr FBuildMergeManifests::MergeDeltaManifest(const FBuildPatchAppManifest& Manifest, const FBuildPatchAppManifest& Delta)
 {
 	using namespace BuildPatchServices;
-	FBuildPatchAppManifestRef MergedManifest = StaticCastSharedRef<FBuildPatchAppManifest>(Manifest->Duplicate());
+	FBuildPatchAppManifestRef MergedManifest = StaticCastSharedRef<FBuildPatchAppManifest>(Manifest.Duplicate());
 	for (FFileManifest& FileManifest : MergedManifest->FileManifestList.FileList)
 	{
-		const FFileManifest* DeltaFileManifest = Delta->GetFileManifest(FileManifest.Filename);
+		const FFileManifest* DeltaFileManifest = Delta.GetFileManifest(FileManifest.Filename);
 		if (DeltaFileManifest != nullptr)
 		{
 			FileManifest.ChunkParts = DeltaFileManifest->ChunkParts;
+			for (auto& Chunk : FileManifest.ChunkParts)
+			{
+				Chunk.bChunkFromDelta = true;
+			}
 		}
 	}
 	if (MergeHelpers::ReinitialiseChunkInfoList(MergedManifest->FileManifestList.FileList, Delta, Manifest, MergedManifest->ChunkDataList.ChunkList))
diff --git a/Engine/Source/Runtime/Online/BuildPatchServices/Private/BuildPatchMergeManifests.h b/Engine/Source/Runtime/Online/BuildPatchServices/Private/BuildPatchMergeManifests.h
--- a/Engine/Source/Runtime/Online/BuildPatchServices/Private/BuildPatchMergeManifests.h
+++ b/Engine/Source/Runtime/Online/BuildPatchServices/Private/BuildPatchMergeManifests.h
@@ -7,7 +7,9 @@
 class FBuildMergeManifests
 {
 public:
-	static bool MergeManifests(const FString& ManifestFilePathA, const FString& ManifestFilePathB, const FString& ManifestFilePathC, const FString& NewVersionString, const FString& SelectionDetailFilePath);
+	static bool MergeManifests(const FString& ManifestFilePathA, const FString& ManifestFilePathB, const FString& ManifestFilePathC, const FString& NewVersionString, const FString& SelectionDetailFilePath, const TMap<FGuid, TArray<uint8>>& AvailableEncryptionSecrets);
 
-	static FBuildPatchAppManifestPtr MergeDeltaManifest(const FBuildPatchAppManifestRef& Manifest, const FBuildPatchAppManifestRef& Delta);
+	static IBuildManifestPtr MergeManifests(const IBuildManifestRef& ManifestA, const IBuildManifestRef& ManifestB, const FString& NewVersionString, const TSet<FString>& FilesFromA, const TSet<FString>& FilesFromB);
+
+	static FBuildPatchAppManifestPtr MergeDeltaManifest(const FBuildPatchAppManifest& Manifest, const FBuildPatchAppManifest& Delta);
 };
diff --git a/Engine/Source/Runtime/Online/BuildPatchServices/Private/BuildPatchProgress.cpp b/Engine/Source/Runtime/Online/BuildPatchServices/Private/BuildPatchProgress.cpp
--- a/Engine/Source/Runtime/Online/BuildPatchServices/Private/BuildPatchProgress.cpp
+++ b/Engine/Source/Runtime/Online/BuildPatchServices/Private/BuildPatchProgress.cpp
@@ -49,6 +49,11 @@ namespace BuildPatchServices
 		bShouldAbort = true;
 	}
 
+	void FBuildPatchProgress::Reset()
+	{
+		bShouldAbort = false;
+	}
+
 	void FBuildPatchProgress::SetStateProgress(const EBuildPatchState& State, const float& Value)
 	{
 		check(!FMath::IsNaN(Value));
@@ -155,11 +160,6 @@ namespace BuildPatchServices
 		}
 	}
 
-	void FBuildPatchProgress::CancelAbort()
-	{
-		bShouldAbort = false;
-	}
-
 	void FBuildPatchProgress::UpdateProgressInfo()
 	{
 		FScopeLock ScopeLock(&ThreadLock);
@@ -220,6 +220,7 @@ namespace BuildPatchServices
 
 	/* FBuildPatchProgress static constants
 	*****************************************************************************/
+	static_assert((uint32)BuildPatchServices::EBuildPatchState::NUM_PROGRESS_STATES == 13, "Please add support for the extra EBuildPatchState values to bHasProgressValue.");
 	const bool FBuildPatchProgress::bHasProgressValue[static_cast<int32>(EBuildPatchState::NUM_PROGRESS_STATES)] =
 	{
 		false, // Queued
@@ -227,15 +228,18 @@ namespace BuildPatchServices
 		true,  // Resuming
 		true,  // Downloading
 		true,  // Installing
+		true,  // RunningCustomUninstallAction
 		true,  // MovingToInstall
 		true,  // SettingAttributes
 		true,  // BuildVerification
 		false, // CleanUp
 		false, // PrerequisitesInstall
-		false, // Completed
+		true,  // Completed
 		true   // Paused
 	};
 
+
+	static_assert((uint32)BuildPatchServices::EBuildPatchState::NUM_PROGRESS_STATES == 13, "Please add support for the extra EBuildPatchState values to bCountsTowardsProgress.");
 	const bool FBuildPatchProgress::bCountsTowardsProgress[static_cast<int32>(EBuildPatchState::NUM_PROGRESS_STATES)] =
 	{
 		false, // Queued
@@ -243,6 +247,7 @@ namespace BuildPatchServices
 		false, // Resuming
 		true,  // Downloading
 		true,  // Installing
+		false, // RunningCustomUninstallAction
 		true,  // MovingToInstall
 		true,  // SettingAttributes
 		true,  // BuildVerification
diff --git a/Engine/Source/Runtime/Online/BuildPatchServices/Private/BuildPatchProgress.h b/Engine/Source/Runtime/Online/BuildPatchServices/Private/BuildPatchProgress.h
--- a/Engine/Source/Runtime/Online/BuildPatchServices/Private/BuildPatchProgress.h
+++ b/Engine/Source/Runtime/Online/BuildPatchServices/Private/BuildPatchProgress.h
@@ -30,6 +30,7 @@ namespace BuildPatchServices
 		// IControllable interface begin.
 		virtual void SetPaused(bool bIsPaused) override;
 		virtual void Abort() override;
+		virtual void Reset() override;
 		// IControllable interface end.
 
 		/**
@@ -102,11 +103,6 @@ namespace BuildPatchServices
 		 */
 		virtual void SetIsDownloading(bool bIsDownloading);
 
-		/**
-		 * Resets an abort state, for the next installer iteration.
-		 */
-		virtual void CancelAbort();
-
 	private:
 
 		/**
diff --git a/Engine/Source/Runtime/Online/BuildPatchServices/Private/BuildPatchServicesModule.cpp b/Engine/Source/Runtime/Online/BuildPatchServices/Private/BuildPatchServicesModule.cpp
--- a/Engine/Source/Runtime/Online/BuildPatchServices/Private/BuildPatchServicesModule.cpp
+++ b/Engine/Source/Runtime/Online/BuildPatchServices/Private/BuildPatchServicesModule.cpp
@@ -21,6 +21,7 @@
 #include "Enumeration/PatchDataEnumeration.h"
 #include "Generation/ChunkDeltaOptimiser.h"
 #include "Generation/PackageChunkData.h"
+#include "Generation/ChunkHarvester.h"
 #include "Installer/BuildStatistics.h"
 #include "Installer/InstallerSharedContext.h"
 #include "Installer/MachineConfig.h"
@@ -121,6 +122,9 @@ IBuildInstallerRef FBuildPatchServicesModule::CreateBuildInstaller(BuildPatchSer
 	{
 		Configuration.bRunRequiredPrereqs = false;
 	}
+	TMap<FGuid, TArray<uint8>> InstallerSecrets = GetEncryptionSecrets();
+	InstallerSecrets.Append(MoveTemp(Configuration.EncryptionSecrets));
+	Configuration.EncryptionSecrets = MoveTemp(InstallerSecrets);
 	FBuildPatchInstallerRef Installer = MakeShared<FBuildPatchInstaller>(MoveTemp(Configuration), AvailableInstallations, LocalMachineConfigFile, Analytics, InstallerStartDelegate, MoveTemp(CompleteDelegate));
 	return Installer;
 }
@@ -136,22 +140,64 @@ IBuildStatisticsRef FBuildPatchServicesModule::CreateBuildStatistics(const IBuil
 	return MakeShareable(FBuildStatisticsFactory::Create(StaticCastSharedRef<FBuildPatchInstaller>(Installer)));
 }
 
+IDirectoryChunkerRef FBuildPatchServicesModule::CreateDirectoryChunker(BuildPatchServices::FDirectoryChunkerConfiguration Configuration) const
+{
+	using namespace BuildPatchServices;
+	return MakeShareable(FDirectoryChunkerFactory::Create(MoveTemp(Configuration)));
+}
+
 IPatchDataEnumerationRef FBuildPatchServicesModule::CreatePatchDataEnumeration(BuildPatchServices::FPatchDataEnumerationConfiguration Configuration) const
 {
 	using namespace BuildPatchServices;
 	return MakeShareable(FPatchDataEnumerationFactory::Create(MoveTemp(Configuration)));
 }
 
+IChunkDeltaOptimiserRef FBuildPatchServicesModule::CreateChunkDeltaOptimiser(BuildPatchServices::FChunkDeltaOptimiserConfiguration Configuration) const
+{
+	using namespace BuildPatchServices;
+	TMap<FGuid, TArray<uint8>> AllSecrets = GetEncryptionSecrets();
+	AllSecrets.Append(MoveTemp(Configuration.EncryptionSecrets));
+	Configuration.EncryptionSecrets = MoveTemp(AllSecrets);
+	return MakeShareable(FChunkDeltaOptimiserFactory::Create(MoveTemp(Configuration)));
+}
+
+IChunkHarvesterRef FBuildPatchServicesModule::CreateChunkHarvester(BuildPatchServices::FChunkHarvesterConfiguration Configuration) const
+{
+	using namespace BuildPatchServices;
+	return MakeShareable(FChunkHarvesterFactory::Create(MoveTemp(Configuration)));
+}
+
+void FBuildPatchServicesModule::DecryptManifestIfNecessary(FBuildPatchAppManifest& Manifest) const
+{
+	bool bIsManifestEncrypted = false;
+	const FGuid EncryptionSecretId = Manifest.GetEncryptionSecretId(&bIsManifestEncrypted);
+	if (!bIsManifestEncrypted)
+	{
+		return;
+	}
+
+	const TOptional<TArray<uint8>> EncryptionSecret = FindEncryptionSecret(EncryptionSecretId);
+	if (!EncryptionSecret.IsSet())
+	{
+		return;
+	}
+
+	const TMap<FGuid, TArray<uint8>> EncryptionSecrets = { {EncryptionSecretId, *EncryptionSecret} };
+	Manifest.DecryptData(EncryptionSecrets);
+}
+
 IBuildManifestPtr FBuildPatchServicesModule::LoadManifestFromFile( const FString& Filename )
 {
 	FBuildPatchAppManifestRef Manifest = MakeShareable( new FBuildPatchAppManifest() );
 	if( Manifest->LoadFromFile( Filename ) )
 	{
+		DecryptManifestIfNecessary(*Manifest);
 		return Manifest;
 	}
 	else
 	{
-		return NULL;
+		GLog->Log(FString::Printf(TEXT("BuildPatchServicesModule: '%s' file is absent or can't be loaded."), *Filename));
+		return nullptr;
 	}
 }
 
@@ -160,9 +206,10 @@ IBuildManifestPtr FBuildPatchServicesModule::MakeManifestFromData(const TArray<u
 	FBuildPatchAppManifestRef Manifest = MakeShareable(new FBuildPatchAppManifest());
 	if (Manifest->DeserializeFromData(ManifestData))
 	{
+		DecryptManifestIfNecessary(*Manifest);
 		return Manifest;
 	}
-	return NULL;
+	return nullptr;
 }
 
 PRAGMA_DISABLE_DEPRECATION_WARNINGS
@@ -179,7 +226,7 @@ PRAGMA_ENABLE_DEPRECATION_WARNINGS
 
 bool FBuildPatchServicesModule::SaveManifestToFile(const FString& Filename, IBuildManifestRef Manifest)
 {
-	return StaticCastSharedRef<FBuildPatchAppManifest>(Manifest)->SaveToFile(Filename);
+	return StaticCastSharedRef<FBuildPatchAppManifest>(Manifest)->SaveToFile(Filename, Manifest->GetFeatureLevel());
 }
 
 TSet<FString> FBuildPatchServicesModule::GetInstalledPrereqIds() const
@@ -224,18 +271,6 @@ bool FBuildPatchServicesModule::Tick(float Delta)
 	return true;
 }
 
-bool FBuildPatchServicesModule::ChunkBuildDirectory(const BuildPatchServices::FChunkBuildConfiguration& Configuration)
-{
-	return FBuildDataGenerator::ChunkBuildDirectory(Configuration);
-}
-
-bool FBuildPatchServicesModule::OptimiseChunkDelta(const BuildPatchServices::FChunkDeltaOptimiserConfiguration& Configuration)
-{
-	using namespace BuildPatchServices;
-	TUniquePtr<IChunkDeltaOptimiser> ChunkDeltaOptimiser(FChunkDeltaOptimiserFactory::Create(Configuration));
-	return ChunkDeltaOptimiser->Run();
-}
-
 bool FBuildPatchServicesModule::CompactifyCloudDirectory(const BuildPatchServices::FCompactifyConfiguration& Configuration)
 {
 	using namespace BuildPatchServices;
@@ -243,13 +278,6 @@ bool FBuildPatchServicesModule::CompactifyCloudDirectory(const BuildPatchService
 	return PatchDataCompactifier->Run();
 }
 
-bool FBuildPatchServicesModule::EnumeratePatchData(const BuildPatchServices::FPatchDataEnumerationConfiguration& Configuration)
-{
-	using namespace BuildPatchServices;
-	TUniquePtr<IPatchDataEnumeration> PatchDataEnumeration(FPatchDataEnumerationFactory::Create(Configuration));
-	return PatchDataEnumeration->Run();
-}
-
 bool FBuildPatchServicesModule::VerifyChunkData(const FString& SearchPath, const FString& OutputFile)
 {
 	return FBuildVerifyChunkData::VerifyChunkData(SearchPath, OutputFile);
@@ -264,7 +292,12 @@ bool FBuildPatchServicesModule::PackageChunkData(const BuildPatchServices::FPack
 
 bool FBuildPatchServicesModule::MergeManifests(const FString& ManifestFilePathA, const FString& ManifestFilePathB, const FString& ManifestFilePathC, const FString& NewVersionString, const FString& SelectionDetailFilePath)
 {
-	return FBuildMergeManifests::MergeManifests(ManifestFilePathA, ManifestFilePathB, ManifestFilePathC, NewVersionString, SelectionDetailFilePath);
+	return FBuildMergeManifests::MergeManifests(ManifestFilePathA, ManifestFilePathB, ManifestFilePathC, NewVersionString, SelectionDetailFilePath, GetEncryptionSecrets());
+}
+
+IBuildManifestPtr FBuildPatchServicesModule::MergeManifests(const IBuildManifestRef& ManifestA, const IBuildManifestRef& ManifestB, const FString& NewVersionString, const TSet<FString>& FilesFromA, const TSet<FString>& FilesFromB)
+{
+	return FBuildMergeManifests::MergeManifests(ManifestA, ManifestB, NewVersionString, FilesFromA, FilesFromB);
 }
 
 bool FBuildPatchServicesModule::DiffManifests(const BuildPatchServices::FDiffManifestsConfiguration& Configuration)
@@ -338,13 +371,50 @@ bool FBuildPatchServicesModule::UnregisterAppInstallation(FString AppInstallDire
 {
 	FPaths::NormalizeDirectoryName(AppInstallDirectory);
 	FPaths::CollapseRelativeDirectories(AppInstallDirectory);
-	if (AvailableInstallations.Remove(AppInstallDirectory) == 1)
+	if (AvailableInstallations.Remove(AppInstallDirectory) > 0)
+	{
+		return true;
+	}
+	return false;
+}
+
+void FBuildPatchServicesModule::RegisterEncryptionSecret(const FGuid& EncryptionSecretId, TArray<uint8> EncryptionSecretKey)
+{
+	FScopeLock ScopeLock(&AvailableEncryptionSecretsCS);
+	AvailableEncryptionSecrets.Add(EncryptionSecretId, MoveTemp(EncryptionSecretKey));
+}
+
+bool FBuildPatchServicesModule::UnregisterEncryptionSecret(const FGuid& EncryptionSecretId)
+{
+	FScopeLock ScopeLock(&AvailableEncryptionSecretsCS);
+	if (AvailableEncryptionSecrets.Remove(EncryptionSecretId) > 0)
 	{
 		return true;
 	}
 	return false;
 }
 
+TOptional<TArray<uint8>> FBuildPatchServicesModule::FindEncryptionSecret(const FGuid& EncryptionSecretId) const
+{
+	FScopeLock ScopeLock(&AvailableEncryptionSecretsCS);
+	const TArray<uint8>* EncryptionSecret = AvailableEncryptionSecrets.Find(EncryptionSecretId);
+	return EncryptionSecret ? TOptional<TArray<uint8>>{*EncryptionSecret} : TOptional<TArray<uint8>>{};
+}
+
+TMap<FGuid, TArray<uint8>> FBuildPatchServicesModule::GetEncryptionSecrets() const
+{
+	FScopeLock ScopeLock(&AvailableEncryptionSecretsCS);
+	return AvailableEncryptionSecrets;
+}
+
+
+bool FBuildPatchServicesModule::CreateRandomBytes(const TArrayView<uint8> OutData)
+{
+	TUniquePtr<ICrypto> Crypto{ FCryptoFactory::Create() };
+
+	return Crypto->CreateRandomBytes(OutData);
+}
+
 void FBuildPatchServicesModule::CancelAllInstallers(bool WaitForThreads)
 {
 	// Using a local bool for this check will improve the assert message that gets displayed
diff --git a/Engine/Source/Runtime/Online/BuildPatchServices/Private/BuildPatchServicesModule.h b/Engine/Source/Runtime/Online/BuildPatchServices/Private/BuildPatchServicesModule.h
--- a/Engine/Source/Runtime/Online/BuildPatchServices/Private/BuildPatchServicesModule.h
+++ b/Engine/Source/Runtime/Online/BuildPatchServices/Private/BuildPatchServicesModule.h
@@ -4,6 +4,7 @@
 
 #include "CoreMinimal.h"
 #include "Misc/Guid.h"
+#include "Misc/Optional.h"
 #include "Interfaces/IBuildPatchServicesModule.h"
 #include "BuildPatchInstaller.h"
 #include "Containers/Ticker.h"
@@ -38,7 +39,10 @@ public:
 	virtual IBuildInstallerRef CreateBuildInstaller(BuildPatchServices::FBuildInstallerConfiguration Configuration, FBuildPatchInstallerDelegate OnComplete) const override;
 	virtual IBuildInstallerSharedContextRef CreateBuildInstallerSharedContext(const TCHAR* DebugName) const override;
 	virtual BuildPatchServices::IBuildStatisticsRef CreateBuildStatistics(const IBuildInstallerRef& Installer) const override;
+	virtual BuildPatchServices::IDirectoryChunkerRef CreateDirectoryChunker(BuildPatchServices::FDirectoryChunkerConfiguration Configuration) const override;
 	virtual BuildPatchServices::IPatchDataEnumerationRef CreatePatchDataEnumeration(BuildPatchServices::FPatchDataEnumerationConfiguration Configuration) const override;
+	virtual BuildPatchServices::IChunkDeltaOptimiserRef CreateChunkDeltaOptimiser(BuildPatchServices::FChunkDeltaOptimiserConfiguration Configuration) const override;
+	virtual BuildPatchServices::IChunkHarvesterRef CreateChunkHarvester(BuildPatchServices::FChunkHarvesterConfiguration Configuration) const override;
 	virtual IBuildManifestPtr LoadManifestFromFile(const FString& Filename) override;
 	virtual IBuildManifestPtr MakeManifestFromData(const TArray<uint8>& ManifestData) override;
 	virtual bool SaveManifestToFile(const FString& Filename, IBuildManifestRef Manifest) override;
@@ -51,14 +55,15 @@ public:
 	virtual void SetAnalyticsProvider(TSharedPtr< IAnalyticsProvider > AnalyticsProvider) override;
 	virtual void RegisterAppInstallation(IBuildManifestRef AppManifest, const FString AppInstallDirectory) override;
 	virtual bool UnregisterAppInstallation(const FString AppInstallDirectory) override;
+	virtual void RegisterEncryptionSecret(const FGuid& EncryptionSecretId, TArray<uint8> EncryptionSecretKey) override;
+	virtual bool UnregisterEncryptionSecret(const FGuid& EncryptionSecretId) override;
+	virtual bool CreateRandomBytes(const TArrayView<uint8> OutData) override;
 	virtual void CancelAllInstallers(bool WaitForThreads) override;
-	virtual bool ChunkBuildDirectory(const BuildPatchServices::FChunkBuildConfiguration& Configuration) override;
-	virtual bool OptimiseChunkDelta(const BuildPatchServices::FChunkDeltaOptimiserConfiguration& Configuration) override;
 	virtual bool CompactifyCloudDirectory(const BuildPatchServices::FCompactifyConfiguration& Configuration) override;
-	virtual bool EnumeratePatchData(const BuildPatchServices::FPatchDataEnumerationConfiguration& Configuration) override;
 	virtual bool VerifyChunkData(const FString& SearchPath, const FString& OutputFile) override;
 	virtual bool PackageChunkData(const BuildPatchServices::FPackageChunksConfiguration& Configuration) override;
 	virtual bool MergeManifests(const FString& ManifestFilePathA, const FString& ManifestFilePathB, const FString& ManifestFilePathC, const FString& NewVersionString, const FString& SelectionDetailFilePath) override;
+	virtual IBuildManifestPtr MergeManifests(const IBuildManifestRef& ManifestA, const IBuildManifestRef& ManifestB, const FString& NewVersionString, const TSet<FString>& FilesFromA, const TSet<FString>& FilesFromB) override;
 	virtual bool DiffManifests(const BuildPatchServices::FDiffManifestsConfiguration& Configuration) override;
 	virtual FSimpleEvent& OnStartBuildInstall() override;
 	virtual IBuildManifestPtr MakeManifestFromJSON(const FString& ManifestJSON) override;
@@ -112,6 +117,15 @@ private:
 	 */
 	void NormalizeCloudPaths(TArray<FString>& InOutCloudPaths);
 
+	// Get the encryption secret by the specified secret ID.
+	TOptional<TArray<uint8>> FindEncryptionSecret(const FGuid& EncryptionSecretId) const;
+
+	// Get all registered encryption secrets.
+	TMap<FGuid, TArray<uint8>> GetEncryptionSecrets() const;
+
+	// Decrypt manifest when the manifest's encryption secret is available
+	void DecryptManifestIfNecessary(FBuildPatchAppManifest& Manifest) const;
+
 private:
 	// The analytics provider interface
 	static TSharedPtr<IAnalyticsProvider> Analytics;
@@ -140,6 +154,12 @@ private:
 	// Holds available installations used for recycling install data
 	TMultiMap<FString, FBuildPatchAppManifestRef> AvailableInstallations;
 
+	// Critical section for multi-threaded access to AvailableEncryptionSecrets
+	mutable FCriticalSection AvailableEncryptionSecretsCS;
+
+	// Holds available encryption secrets used for decrypting installation data
+	TMap<FGuid, TArray<uint8>> AvailableEncryptionSecrets;
+
 	// Handle to the registered Tick delegate
 	FTSTicker::FDelegateHandle TickDelegateHandle;
 
