p4 integrate //UE5/Dev-Distro-5.5/Engine/Source/Programs/BuildPatchTool/... //UE5/Main/Engine/Source/Programs/BuildPatchTool/...

diff --git a/Engine/Source/Runtime/Online/BuildPatchServices/Private/BuildPatchSettings.cpp b/Engine/Source/Runtime/Online/BuildPatchServices/Private/BuildPatchSettings.cpp
--- a/Engine/Source/Runtime/Online/BuildPatchServices/Private/BuildPatchSettings.cpp
+++ b/Engine/Source/Runtime/Online/BuildPatchServices/Private/BuildPatchSettings.cpp
@@ -19,12 +19,13 @@ namespace BuildPatchServices
 		Install,
 		Update,
 		Repair,
+		VerifyOnly,
 		Uninstall,
 
 		Invalid
 	};
 
-	BuildPatchServices::FInstallerAction FInstallerAction::MakeInstall(const IBuildManifestRef& Manifest, TSet<FString> InstallTags /*= TSet<FString>()*/, FString InstallSubdirectory /*= FString()*/, FString CloudSubdirectory /*= FString()*/)
+	FInstallerAction FInstallerAction::MakeInstall(const IBuildManifestRef& Manifest, TSet<FString> InstallTags /*= TSet<FString>()*/, FString InstallSubdirectory /*= FString()*/, FString CloudSubdirectory /*= FString()*/)
 	{
 		FInstallerAction InstallerAction;
 		InstallerAction.InstallManifest = Manifest;
@@ -35,7 +36,7 @@ namespace BuildPatchServices
 		return InstallerAction;
 	}
 
-	BuildPatchServices::FInstallerAction FInstallerAction::MakeUpdate(const IBuildManifestRef& CurrentManifest, const IBuildManifestRef& InstallManifest, TSet<FString> InstallTags /*= TSet<FString>()*/, FString InstallSubdirectory /*= FString()*/, FString CloudSubdirectory /*= FString()*/)
+	FInstallerAction FInstallerAction::MakeUpdate(const IBuildManifestRef& CurrentManifest, const IBuildManifestRef& InstallManifest, TSet<FString> InstallTags /*= TSet<FString>()*/, FString InstallSubdirectory /*= FString()*/, FString CloudSubdirectory /*= FString()*/)
 	{
 		FInstallerAction InstallerAction;
 		InstallerAction.CurrentManifest = CurrentManifest;
@@ -47,7 +48,7 @@ namespace BuildPatchServices
 		return InstallerAction;
 	}
 
-	BuildPatchServices::FInstallerAction FInstallerAction::MakeRepair(const IBuildManifestRef& Manifest, TSet<FString> InstallTags /*= TSet<FString>()*/, FString InstallSubdirectory /*= FString()*/, FString CloudSubdirectory /*= FString()*/)
+	FInstallerAction FInstallerAction::MakeRepair(const IBuildManifestRef& Manifest, TSet<FString> InstallTags /*= TSet<FString>()*/, FString InstallSubdirectory /*= FString()*/, FString CloudSubdirectory /*= FString()*/)
 	{
 		FInstallerAction InstallerAction;
 		InstallerAction.CurrentManifest = Manifest;
@@ -59,34 +60,26 @@ namespace BuildPatchServices
 		return InstallerAction;
 	}
 
-	BuildPatchServices::FInstallerAction FInstallerAction::MakeUninstall(const IBuildManifestRef& Manifest, FString InstallSubdirectory /*= FString()*/, FString CloudSubdirectory /*= FString()*/)
+	FInstallerAction FInstallerAction::MakeVerifyOnly(const IBuildManifestRef& Manifest, TSet<FString> InstallTags /*= TSet<FString>()*/, FString InstallSubdirectory /*= FString()*/, FString CloudSubdirectory /*= FString()*/)
 	{
 		FInstallerAction InstallerAction;
 		InstallerAction.CurrentManifest = Manifest;
+		InstallerAction.InstallManifest = Manifest;
+		InstallerAction.InstallTags = MoveTemp(InstallTags);
 		InstallerAction.InstallSubdirectory = MoveTemp(InstallSubdirectory);
 		InstallerAction.CloudSubdirectory = MoveTemp(CloudSubdirectory);
-		InstallerAction.ActionIntent = EInstallActionIntent::Uninstall;
+		InstallerAction.ActionIntent = EInstallActionIntent::VerifyOnly;
 		return InstallerAction;
 	}
 
-	FInstallerAction::FInstallerAction(const FInstallerAction& CopyFrom)
-		: CurrentManifest(CopyFrom.CurrentManifest)
-		, InstallManifest(CopyFrom.InstallManifest)
-		, InstallTags(CopyFrom.InstallTags)
-		, InstallSubdirectory(CopyFrom.InstallSubdirectory)
-		, CloudSubdirectory(CopyFrom.CloudSubdirectory)
-		, ActionIntent(CopyFrom.ActionIntent)
-	{
-	}
-
-	FInstallerAction::FInstallerAction(FInstallerAction&& MoveFrom)
-		: CurrentManifest(MoveTemp(MoveFrom.CurrentManifest))
-		, InstallManifest(MoveTemp(MoveFrom.InstallManifest))
-		, InstallTags(MoveTemp(MoveFrom.InstallTags))
-		, InstallSubdirectory(MoveTemp(MoveFrom.InstallSubdirectory))
-		, CloudSubdirectory(MoveTemp(MoveFrom.CloudSubdirectory))
-		, ActionIntent(MoveFrom.ActionIntent)
+	FInstallerAction FInstallerAction::MakeUninstall(const IBuildManifestRef& Manifest, FString InstallSubdirectory /*= FString()*/, FString CloudSubdirectory /*= FString()*/)
 	{
+		FInstallerAction InstallerAction;
+		InstallerAction.CurrentManifest = Manifest;
+		InstallerAction.InstallSubdirectory = MoveTemp(InstallSubdirectory);
+		InstallerAction.CloudSubdirectory = MoveTemp(CloudSubdirectory);
+		InstallerAction.ActionIntent = EInstallActionIntent::Uninstall;
+		return InstallerAction;
 	}
 
 	FInstallerAction::FInstallerAction()
@@ -109,6 +102,11 @@ namespace BuildPatchServices
 		return ActionIntent == EInstallActionIntent::Repair;
 	}
 
+	bool FInstallerAction::IsVerifyOnly() const
+	{
+		return ActionIntent == EInstallActionIntent::VerifyOnly;
+	}
+
 	bool FInstallerAction::IsUninstall() const
 	{
 		return ActionIntent == EInstallActionIntent::Uninstall;
@@ -141,35 +139,13 @@ namespace BuildPatchServices
 
 	FBuildInstallerConfiguration::FBuildInstallerConfiguration(TArray<FInstallerAction> InInstallerActions)
 		: InstallerActions(MoveTemp(InInstallerActions))
-		, InstallDirectory()
-		, StagingDirectory()
-		, BackupDirectory()
-		, ChunkDatabaseFiles()
-		, CloudDirectories()
-		, InstallMode(EInstallMode::NonDestructiveInstall)
-		, VerifyMode(EVerifyMode::ShaVerifyAllFiles)
-		, DeltaPolicy(EDeltaPolicy::Skip)
-		, bRunRequiredPrereqs(true)
-		, bAllowConcurrentExecution(false)
-		, bTrackFileOperations(true)
-	{
-	}
-
-	FChunkBuildConfiguration::FChunkBuildConfiguration()
-		: FeatureLevel(EFeatureLevel::Latest)
-		, AppId(0)
-		, DataAgeThreshold(TNumericLimits<float>::Max())
-		, bShouldHonorReuseThreshold(false)
-		, OutputChunkWindowSize(1024 * 1024)
-		, bShouldMatchAnyWindowSize(true)
-		, bAllowEmptyBuild(false)
 	{
 	}
 
 	FChunkDeltaOptimiserConfiguration::FChunkDeltaOptimiserConfiguration()
 		: ScanWindowSize(8191)
 		, OutputChunkSize(1024 * 1024)
-		, DiffAbortThreshold(TNumericLimits<int64>::Max())
+		, bResaveExistingChunks(false)
 	{
 	}
 
@@ -178,10 +154,6 @@ namespace BuildPatchServices
 	{
 	}
 
-	FDiffManifestsConfiguration::FDiffManifestsConfiguration()
-	{
-	}
-
 	FCompactifyConfiguration::FCompactifyConfiguration()
 		: DataAgeThreshold(7.0f)
 		, bRunPreview(true)
@@ -195,56 +167,13 @@ namespace BuildPatchServices
 	}
 
 	FInstallerConfiguration::FInstallerConfiguration(const IBuildManifestRef& InInstallManifest)
-		: CurrentManifest(nullptr)
-		, InstallManifest(InInstallManifest)
-		, InstallDirectory()
-		, StagingDirectory()
-		, BackupDirectory()
-		, ChunkDatabaseFiles()
-		, CloudDirectories()
-		, InstallTags()
-		, InstallMode(EInstallMode::NonDestructiveInstall)
-		, VerifyMode(EVerifyMode::ShaVerifyAllFiles)
-		, DeltaPolicy(EDeltaPolicy::Skip)
-		, bIsRepair(false)
-		, bRunRequiredPrereqs(true)
-		, bAllowConcurrentExecution(false)
-	{
-	}
-
-	FInstallerConfiguration::FInstallerConfiguration(FInstallerConfiguration&& MoveFrom)
-		: CurrentManifest(MoveTemp(MoveFrom.CurrentManifest))
-		, InstallManifest(MoveTemp(MoveFrom.InstallManifest))
-		, InstallDirectory(MoveTemp(MoveFrom.InstallDirectory))
-		, StagingDirectory(MoveTemp(MoveFrom.StagingDirectory))
-		, BackupDirectory(MoveTemp(MoveFrom.BackupDirectory))
-		, ChunkDatabaseFiles(MoveTemp(MoveFrom.ChunkDatabaseFiles))
-		, CloudDirectories(MoveTemp(MoveFrom.CloudDirectories))
-		, InstallTags(MoveTemp(MoveFrom.InstallTags))
-		, InstallMode(MoveFrom.InstallMode)
-		, VerifyMode(MoveFrom.VerifyMode)
-		, DeltaPolicy(MoveFrom.DeltaPolicy)
-		, bIsRepair(MoveFrom.bIsRepair)
-		, bRunRequiredPrereqs(MoveFrom.bRunRequiredPrereqs)
-		, bAllowConcurrentExecution(MoveFrom.bAllowConcurrentExecution)
-	{
-	}
-
-	FInstallerConfiguration::FInstallerConfiguration(const FInstallerConfiguration& CopyFrom)
-		: CurrentManifest(CopyFrom.CurrentManifest)
-		, InstallManifest(CopyFrom.InstallManifest)
-		, InstallDirectory(CopyFrom.InstallDirectory)
-		, StagingDirectory(CopyFrom.StagingDirectory)
-		, BackupDirectory(CopyFrom.BackupDirectory)
-		, ChunkDatabaseFiles(CopyFrom.ChunkDatabaseFiles)
-		, CloudDirectories(CopyFrom.CloudDirectories)
-		, InstallTags(CopyFrom.InstallTags)
-		, InstallMode(CopyFrom.InstallMode)
-		, VerifyMode(CopyFrom.VerifyMode)
-		, DeltaPolicy(CopyFrom.DeltaPolicy)
-		, bIsRepair(CopyFrom.bIsRepair)
-		, bRunRequiredPrereqs(CopyFrom.bRunRequiredPrereqs)
-		, bAllowConcurrentExecution(CopyFrom.bAllowConcurrentExecution)
+		: InstallManifest(InInstallManifest)
+	{
+	}
+
+	FChunkHarvesterConfiguration::FChunkHarvesterConfiguration()
+		: FeatureLevelOverride(EFeatureLevel::Invalid)
+		, bResaveExistingChunks(false)
 	{
 	}
 }
diff --git a/Engine/Source/Runtime/Online/BuildPatchServices/Private/BuildPatchState.cpp b/Engine/Source/Runtime/Online/BuildPatchServices/Private/BuildPatchState.cpp
--- a/Engine/Source/Runtime/Online/BuildPatchServices/Private/BuildPatchState.cpp
+++ b/Engine/Source/Runtime/Online/BuildPatchServices/Private/BuildPatchState.cpp
@@ -31,6 +31,7 @@ const FText& BuildPatchServices::StateToText(const BuildPatchServices::EBuildPat
 	case BuildPatchServices::EBuildPatchState::Downloading:
 		return Downloading;
 	case BuildPatchServices::EBuildPatchState::Installing:
+	case BuildPatchServices::EBuildPatchState::RunningCustomUninstallAction:
 	case BuildPatchServices::EBuildPatchState::MovingToInstall:
 	case BuildPatchServices::EBuildPatchState::SettingAttributes:
 		return Installing;
diff --git a/Engine/Source/Runtime/Online/BuildPatchServices/Private/BuildPatchUtil.cpp b/Engine/Source/Runtime/Online/BuildPatchServices/Private/BuildPatchUtil.cpp
--- a/Engine/Source/Runtime/Online/BuildPatchServices/Private/BuildPatchUtil.cpp
+++ b/Engine/Source/Runtime/Online/BuildPatchServices/Private/BuildPatchUtil.cpp
@@ -14,54 +14,87 @@
 
 using namespace BuildPatchServices;
 
-/* FBuildPatchUtils implementation
-*****************************************************************************/
-FString FBuildPatchUtils::GetChunkNewFilename(const EFeatureLevel FeatureLevel, const FGuid& ChunkGUID, const uint64& ChunkHash)
+namespace BuildPatchUtils
 {
-	check(ChunkGUID.IsValid());
-	if (FeatureLevel < EFeatureLevel::DataFileRenames)
+	void B64UriSafe(FString& B64)
+	{
+		// Make URI safe.
+		B64.ReplaceInline(TEXT("+"), TEXT("-"), ESearchCase::CaseSensitive);
+		B64.ReplaceInline(TEXT("/"), TEXT("_"), ESearchCase::CaseSensitive);
+		// Trim = characters.
+		B64.ReplaceInline(TEXT("="), TEXT(""), ESearchCase::CaseSensitive);
+	}
+
+	FString ToMiniUriString(const FGuid& Val)
+	{
+		FString Mini = FBase64::Encode((uint8*)&Val, sizeof(Val));
+		B64UriSafe(Mini);
+		return Mini;
+	}
+
+	bool FromMiniUriString(const FString& String, FGuid& Val)
+	{
+		if (String.Len() == 22)
+		{
+			FString Mini = String.Replace(TEXT("-"), TEXT("+"), ESearchCase::CaseSensitive)
+			                     .Replace(TEXT("_"), TEXT("/"), ESearchCase::CaseSensitive);
+			return FBase64::Decode(Mini.GetCharArray().GetData(), 22, (uint8*)&Val);
+		}
+		return false;
+	}
+
+	FString ToMiniUriString(const uint64& Val)
+	{
+		FString Mini = FBase64::Encode((uint8*)&Val, sizeof(Val));
+		B64UriSafe(Mini);
+		return Mini;
+	}
+
+	int32 DataIdToGroupId(const FGuid& DataId)
 	{
-		return GetChunkOldFilename(ChunkGUID);
+		return FCrc::MemCrc32(&DataId, sizeof(FGuid)) % 100;
+	}
+
+	FString GetChunkNewFilename(BuildPatchServices::EFeatureLevel FeatureLevel, const FGuid& ChunkId, const uint64& RollingHash, const FGuid& EncryptionSecretId)
+	{
+		if (FeatureLevel < EFeatureLevel::DataFileRenames)
+		{
+			return FBuildPatchUtils::GetChunkOldFilename(ChunkId);
+		}
+		else if (FeatureLevel < EFeatureLevel::ChunksStoredBySecret)
+		{
+			return FString::Printf(TEXT("%s/%02d/%016llX_%s.chunk"), ManifestVersionHelpers::GetChunkSubdir(FeatureLevel), DataIdToGroupId(ChunkId), RollingHash, *ChunkId.ToString());
+		}
+		else
+		{
+			const FString SecretPathComponent = EncryptionSecretId.IsValid() ? ToMiniUriString(EncryptionSecretId) : TEXT("plain");
+			return FString::Printf(TEXT("%s/%s/%02d/%s_%s.chunk"), ManifestVersionHelpers::GetChunkSubdir(FeatureLevel), *SecretPathComponent, DataIdToGroupId(ChunkId), *ToMiniUriString(RollingHash), *ToMiniUriString(ChunkId));
+		}
 	}
-	return FString::Printf(TEXT("%s/%02d/%016llX_%s.chunk"), ManifestVersionHelpers::GetChunkSubdir(FeatureLevel), FCrc::MemCrc32(&ChunkGUID, sizeof(FGuid)) % 100, ChunkHash, *ChunkGUID.ToString());
 }
 
-FString FBuildPatchUtils::GetFileNewFilename(const EFeatureLevel FeatureLevel, const FGuid& FileGUID, const FSHAHash& FileHash)
+/* FBuildPatchUtils implementation
+*****************************************************************************/
+FString FBuildPatchUtils::GetChunkNewFilename(BuildPatchServices::EFeatureLevel FeatureLevel, const BuildPatchServices::FChunkInfo& Chunk)
 {
-	check(FileGUID.IsValid());
-	return FString::Printf(TEXT("%s/%02d/%s_%s.file"), ManifestVersionHelpers::GetFileSubdir(FeatureLevel), FCrc::MemCrc32(&FileGUID, sizeof(FGuid)) % 100, *FileHash.ToString(), *FileGUID.ToString());
+	return BuildPatchUtils::GetChunkNewFilename(FeatureLevel, Chunk.Guid, Chunk.Hash, Chunk.EncryptionSecretId);
 }
 
-FString FBuildPatchUtils::GetFileNewFilename(const EFeatureLevel FeatureLevel, const FGuid& FileGUID, const uint64& FileHash)
+FString FBuildPatchUtils::GetChunkNewFilename(BuildPatchServices::EFeatureLevel FeatureLevel, const BuildPatchServices::FChunkHeader& Chunk)
 {
-	check(FileGUID.IsValid());
-	return FString::Printf(TEXT("%s/%02d/%016llX_%s.file"), ManifestVersionHelpers::GetFileSubdir(FeatureLevel), FCrc::MemCrc32(&FileGUID, sizeof(FGuid)) % 100, FileHash, *FileGUID.ToString());
+	return BuildPatchUtils::GetChunkNewFilename(FeatureLevel, Chunk.Guid, Chunk.RollingHash, Chunk.EncryptionSecretId);
 }
 
-void FBuildPatchUtils::GetChunkDetailFromNewFilename( const FString& ChunkNewFilename, FGuid& ChunkGUID, uint64& ChunkHash )
+FString FBuildPatchUtils::GetFileNewFilename(const EFeatureLevel FeatureLevel, const FGuid& FileGUID, const FSHAHash& FileHash)
 {
-	const FString ChunkFilename = FPaths::GetBaseFilename( ChunkNewFilename );
-	FString GuidString;
-	FString HashString;
-	ChunkFilename.Split( TEXT( "_" ), &HashString, &GuidString );
-	// Check that the strings are exactly as we expect otherwise this is not being used properly
-	check( GuidString.Len() == 32 );
-	check( HashString.Len() == 16 );
-	ChunkHash = FCString::Strtoui64( *HashString, NULL, 16 );
-	FGuid::Parse( GuidString, ChunkGUID );
+	check(FileGUID.IsValid());
+	return FString::Printf(TEXT("%s/%02d/%s_%s.file"), ManifestVersionHelpers::GetFileSubdir(FeatureLevel), BuildPatchUtils::DataIdToGroupId(FileGUID), *FileHash.ToString(), *FileGUID.ToString());
 }
 
-void FBuildPatchUtils::GetFileDetailFromNewFilename(const FString& FileNewFilename, FGuid& FileGUID, FSHAHash& FileHash)
+FString FBuildPatchUtils::GetFileNewFilename(const EFeatureLevel FeatureLevel, const FGuid& FileGUID, const uint64& FileHash)
 {
-	const FString FileFilename = FPaths::GetBaseFilename( FileNewFilename );
-	FString GuidString;
-	FString HashString;
-	FileFilename.Split( TEXT( "_" ), &HashString, &GuidString );
-	// Check that the strings are exactly as we expect otherwise this is not being used properly
-	check( GuidString.Len() == 32 );
-	check( HashString.Len() == 40 );
-	HexToBytes( HashString, FileHash.Hash );
-	FGuid::Parse( GuidString, FileGUID );
+	check(FileGUID.IsValid());
+	return FString::Printf(TEXT("%s/%02d/%016llX_%s.file"), ManifestVersionHelpers::GetFileSubdir(FeatureLevel), BuildPatchUtils::DataIdToGroupId(FileGUID), FileHash, *FileGUID.ToString());
 }
 
 FString FBuildPatchUtils::GetChunkOldFilename(const FGuid& ChunkGUID)
@@ -105,13 +138,12 @@ FString FBuildPatchUtils::GetDataFilename(const FBuildPatchAppManifest&    Manif
 	{
 		return FBuildPatchUtils::GetDataTypeOldFilename(DataType, DataGUID);
 	}
-	else if (!Manifest.IsFileDataManifest())
+	else if (DataType == EBuildPatchDataType::ChunkData)
 	{
-		uint64 ChunkHash;
-		const bool bFound = Manifest.GetChunkHash(DataGUID, ChunkHash);
+		const FChunkInfo* ChunkInfo = Manifest.GetChunkInfo(DataGUID);
 		// Should be impossible to not exist
-		check(bFound);
-		return FBuildPatchUtils::GetChunkNewFilename(Manifest.GetFeatureLevel(), DataGUID, ChunkHash);
+		check(ChunkInfo != nullptr);
+		return FBuildPatchUtils::GetChunkNewFilename(Manifest.GetFeatureLevel(), *ChunkInfo);
 	}
 	else if (Manifest.GetFeatureLevel() <= EFeatureLevel::StoredAsCompressedUClass)
 	{
@@ -131,21 +163,19 @@ FString FBuildPatchUtils::GetDataFilename(const FBuildPatchAppManifest&    Manif
 	}
 }
 
-bool FBuildPatchUtils::GetGUIDFromFilename( const FString& DataFilename, FGuid& DataGUID )
+bool FBuildPatchUtils::GetGUIDFromFilename(const FString& DataFilename, FGuid& DataGUID)
 {
-	const FString DataBaseFilename = FPaths::GetBaseFilename( DataFilename );
+	const FString DataBaseFilename = FPaths::GetBaseFilename(DataFilename);
 	FString GuidString;
-	if( DataBaseFilename.Contains( TEXT( "_" ) ) )
-	{
-		DataBaseFilename.Split( TEXT( "_" ), NULL, &GuidString );
-	}
-	else
+	if (DataBaseFilename.Len() == 49)
 	{
-		GuidString = DataBaseFilename;
+		GuidString = DataBaseFilename.Right(32);
+		return FGuid::Parse(GuidString, DataGUID);
 	}
-	if(GuidString.Len() == 32)
+	else if (DataBaseFilename.Len() == 34)
 	{
-		return FGuid::Parse(GuidString, DataGUID);
+		GuidString = DataBaseFilename.Right(22);
+		return BuildPatchUtils::FromMiniUriString(GuidString, DataGUID);
 	}
 	return false;
 }
@@ -198,7 +228,7 @@ FString FBuildPatchUtils::GetChunkDeltaDirectory(const FBuildPatchAppManifest& D
 
 FString FBuildPatchUtils::GetChunkDeltaFilename(const FBuildPatchAppManifest& SourceManifest, const FBuildPatchAppManifest& DestinationManifest, const FString& FilenameTrailer)
 {
-    if (FilenameTrailer.Len())
+	if (FilenameTrailer.Len())
 	{
 		return GetChunkDeltaDirectory(DestinationManifest) / SourceManifest.GetBuildId() + TEXT("-") + FilenameTrailer + TEXT(".delta");
 	}
@@ -259,6 +289,31 @@ uint8 FBuildPatchUtils::VerifyFile(IFileSystem* FileSystem, const FString& FileT
 	return ReturnValue;
 }
 
+void FBuildPatchUtils::SHAToBase32(const FSHAHash& SHA, FString& OutString)
+{
+	// For alphabet we'll use the "Extended hex" set.
+	static const TCHAR* Alphabet = TEXT("0123456789ABCDEFGHIJKLMNOPQRSTUV");
+	OutString.Empty(32);
+	const uint8* Bytes = SHA.Hash;
+	for (int32 Idx = 0; (Idx + 4) < 20; Idx += 5)
+	{
+		const uint64 Data =
+			((uint64)Bytes[Idx]) << 32 |
+			((uint64)Bytes[Idx + 1]) << 24 |
+			((uint64)Bytes[Idx + 2]) << 16 |
+			((uint64)Bytes[Idx + 3]) << 8 |
+			((uint64)Bytes[Idx + 4]);
+
+		// 5 bytes produce 8 chars
+		// We shift 35, 30, 25 and so on to isolate each 5bit group starting from highest.
+		for (int32 Shift = 35; Shift >= 0; Shift -= 5)
+		{
+			uint64 ThisRange = Data >> Shift;
+			OutString += Alphabet[ThisRange & 0x1F];
+		}
+	}
+}
+
 FString FBuildPatchUtils::ResolveInstallationFileName(const FString& BaseInstallationDirectory, const TMap<FString, FString>& PerFileSubdirectories, const FString& BuildFilename)
 {
 	const FString* PerFileSubdir = PerFileSubdirectories.Find(BuildFilename);
@@ -272,60 +327,66 @@ FString FBuildPatchUtils::ResolveInstallationFileName(const FString& BaseInstall
 uint64 FBuildPatchUtils::CalculateDiskSpaceRequirementsWithDeleteDuringInstall(
 	const TArray<FString>& InFilesToConstruct, 
 	int32 InCompletedFileCount, 
+	int64 InProgressFileSize,
 	IBuildManifestSet* InManifestSet,
 	const TArray<uint64>& InChunkDbSizesAtPosition,
-	uint64 InTotalChunkDbSize
-	)
+	const FString& InstallDirectory,
+	EInstallMode InstallMode)
 {
 	uint64 TotalDeletedSize = 0;
 	uint64 TotalWrittenSize = 0;
 
-	uint64 MaxDiskSize = 0;
+	int64 MaxDiskSize = 0;
 	int32 MaxDiskSizeFileIndex = 0;
 
-	// We start off with full chunk db size.
-	uint64 TotalChunkDbSizeAtLastFile = InTotalChunkDbSize;
+	TSet<FString> MissedCurrentFiles; // Files that were not installed, possibly due to filtering by tags.
+	if (InstallMode == EInstallMode::DestructiveInstall)
+	{
+		// This data is used only when we overwrite existing files on the fly.
+		InManifestSet->GetMissedFiles(InstallDirectory, MissedCurrentFiles);
+	}
 
 	for (int32 FileIndex = InCompletedFileCount; FileIndex < InFilesToConstruct.Num(); FileIndex++)
 	{
-		uint64 ChunkDbSize = InChunkDbSizesAtPosition[FileIndex];
-
+		const FString& Filename = InFilesToConstruct[FileIndex];
 		// We've completed this file
-		uint64 NewFileSize = 0;
+		if (const FFileManifest* IncomingFileManifest = InManifestSet->GetNewFileManifest(Filename))
 		{
-			const FFileManifest* IncomingFileManifest = InManifestSet->GetNewFileManifest(InFilesToConstruct[FileIndex]);			
-			if (IncomingFileManifest)
+			TotalWrittenSize += IncomingFileManifest->FileSize;
+			if (FileIndex == InCompletedFileCount && InProgressFileSize > 0)
 			{
-				NewFileSize = IncomingFileManifest->FileSize;
+				// If the first file in the list is only partially completed, adjust TotalWrittenSize accordingly.
+				TotalWrittenSize -= FMath::Min<uint64>(InProgressFileSize, IncomingFileManifest->FileSize);
 			}
 		}
-
-		TotalWrittenSize += NewFileSize;
-
 		// We delete the chunkdbs _after_ we write the output so we can't use this size until 
 		// the next file gets done. Be sure to handle the case where we've deleted so much data
 		// we are below the waterline.
-		if (TotalDeletedSize < (TotalWrittenSize + TotalChunkDbSizeAtLastFile) &&
-			(TotalChunkDbSizeAtLastFile + TotalWrittenSize - TotalDeletedSize) > MaxDiskSize)
+		const uint64 TotalChunkDbSizeAtLastFile = InChunkDbSizesAtPosition[FileIndex];
+		const int64 CurrentDiskSize = static_cast<int64>(TotalChunkDbSizeAtLastFile + TotalWrittenSize) - static_cast<int64>(TotalDeletedSize);
+		if (CurrentDiskSize > MaxDiskSize)
 		{
-			MaxDiskSize = TotalChunkDbSizeAtLastFile + TotalWrittenSize - TotalDeletedSize;
+			MaxDiskSize = CurrentDiskSize;
 			MaxDiskSizeFileIndex = FileIndex;
 		}
 
 		//UE_LOG(LogTemp, Display, TEXT("...@ file %d chunks = %llu, install = %llu"), FileIndex, TotalChunkDbSizeAtLastFile, TotalWrittenSize);
 
-		// If we are patching, we can now delete the output file, which decreases our disk presence, however
-		// we update this after the check because we can't delete until after we have the file fully constructed.
-		const FFileManifest* OnDiskFileManifest = InManifestSet->GetCurrentFileManifest(InFilesToConstruct[FileIndex]);
-		if (OnDiskFileManifest)
+		// We should take old files into account only if we will overwrite them on the fly.
+		if (InstallMode == EInstallMode::DestructiveInstall)
 		{
-			TotalDeletedSize += OnDiskFileManifest->FileSize;
+			// If we are patching, we can now delete the output file, which decreases our disk presence, however
+			// we update this after the check because we can't delete until after we have the file fully constructed.
+			const FFileManifest* OnDiskFileManifest = InManifestSet->GetCurrentFileManifest(Filename);
+			// If old manifest contains file and file was not missed from disk due to filtering by tags
+			if (OnDiskFileManifest && !MissedCurrentFiles.Contains(Filename))
+			{
+				TotalDeletedSize += OnDiskFileManifest->FileSize;
+			}
 		}
-
-		TotalChunkDbSizeAtLastFile = ChunkDbSize;
 	}
 
 	//UE_LOG(LogTemp, Display, TEXT("Max disk use %llu (+%llu after d/l) (install size: %llu: +%llu) after file %d"), MaxDiskSize, PostDlSize, TotalWrittenSize, MaxDiskSize - TotalWrittenSize, MaxDiskSizeFileIndex);
 
-	return MaxDiskSize;
+	return static_cast<uint64>(MaxDiskSize);
 }
\ No newline at end of file
diff --git a/Engine/Source/Runtime/Online/BuildPatchServices/Private/BuildPatchUtil.h b/Engine/Source/Runtime/Online/BuildPatchServices/Private/BuildPatchUtil.h
--- a/Engine/Source/Runtime/Online/BuildPatchServices/Private/BuildPatchUtil.h
+++ b/Engine/Source/Runtime/Online/BuildPatchServices/Private/BuildPatchUtil.h
@@ -8,6 +8,7 @@
 
 #include "CoreMinimal.h"
 #include "BuildPatchManifest.h"
+#include "BuildPatchInstall.h"
 
 namespace BuildPatchServices
 {
@@ -32,19 +33,28 @@ DECLARE_DELEGATE_OneParam(FBuildPatchFloatDelegate, float);
 // A delegate returning a bool. Used to pass a paused state.
 DECLARE_DELEGATE_RetVal(bool, FBuildPatchBoolRetDelegate);
 
+/**
+ * Some constants
+ */
+namespace BuildPatchConstants
+{
+	// Result of converting SHA1 to base32 string value length
+	static const int32 SHAToBase32StringLength = 32;
+}
+
 /**
  * Some utility functions
  */
 struct FBuildPatchUtils
 {
 	/**
-	 * Gets the the relative cloud filename for a chunk generated from it's GUID and Hash, which is the new format.
+	 * Gets the relative cloud filename for a chunk in the newest format.
 	 * @param FeatureLevel   The manifest version that references this chunk.
-	 * @param ChunkGuid      The chunk Guid.
-	 * @param ChunkHash      The chunk rolling hash value.
+	 * @param Chunk          The chunk we need a path for
 	 * @return the chunk path.
 	 */
-	static FString GetChunkNewFilename(BuildPatchServices::EFeatureLevel FeatureLevel, const FGuid& ChunkGUID, const uint64& ChunkHash);
+	static FString GetChunkNewFilename(BuildPatchServices::EFeatureLevel FeatureLevel, const BuildPatchServices::FChunkInfo& Chunk);
+	static FString GetChunkNewFilename(BuildPatchServices::EFeatureLevel FeatureLevel, const BuildPatchServices::FChunkHeader& Chunk);
 
 	/**
 	 * Gets the the relative cloud filename for a file chunk generated from it's GUID and Hash, which is the new format.
@@ -56,22 +66,6 @@ struct FBuildPatchUtils
 	static FString GetFileNewFilename(BuildPatchServices::EFeatureLevel FeatureLevel, const FGuid& FileGUID, const FSHAHash& FileHash);
 	static FString GetFileNewFilename(BuildPatchServices::EFeatureLevel FeatureLevel, const FGuid& FileGUID, const uint64& FilePartHash);
 
-	/**
-	 * Gets the chunk GUID and Hash, from the filename, which is the new format.
-	 * @param ChunkNewFilename	IN	The chunk filename that would have been generated from GetChunkNewFilename
-	 * @param ChunkGuid			OUT	The chunk Guid
-	 * @param ChunkHash			OUT	The chunk rolling hash value
-	 */
-	static void GetChunkDetailFromNewFilename(const FString& ChunkNewFilename, FGuid& ChunkGUID, uint64& ChunkHash);
-
-	/**
-	 * Gets the file chunk GUID and file hash, from the filename, which is the new format.
-	 * @param FileNewFilename	IN	The file chunk filename that would have been generated from GetChunkNewFilename
-	 * @param FileGUID			OUT	The file chunk Guid
-	 * @param FileHash			OUT	The file chunk rolling hash value
-	 */
-	static void GetFileDetailFromNewFilename(const FString& FileNewFilename, FGuid& FileGUID, FSHAHash& FileHash);
-
 	/**
 	 * Gets the the relative cloud filename for a chunk generated from it's GUID
 	 * @param ChunkGuid			The chunk Guid
@@ -136,7 +130,7 @@ struct FBuildPatchUtils
 	 * Based on the source and destination manifests, get the filename for the delta that optimises patching from source to destination.
 	 * @param SourceManifest        The source manifest.
 	 * @param DestinationManifest   The destination manifest.
-	 * @param FilenameTrailer		If a nonstandard optimized delta is used, this is the trailer to add to the filename.
+	 * @param FilenameTrailer       This is an optional parameter to allow for multiple optimized deltas for the same build pairing. If unneeded, leave as "".
 	 * @return the CloudDir relative delta filename.
 	 */
 	static FString GetChunkDeltaFilename(const FBuildPatchAppManifest& SourceManifest, const FBuildPatchAppManifest& DestinationManifest, const FString& FilenameTrailer);
@@ -165,13 +159,22 @@ struct FBuildPatchUtils
 	 * @return		0 if no match, 1 for match with Hash1, and 2 for match with Hash2
 	 */
 	static uint8 VerifyFile(BuildPatchServices::IFileSystem* FileSystem, const FString& FileToVerify, const FSHAHash& Hash1, const FSHAHash& Hash2);
+	
+	/**
+	 * Encodes SHA hash to base32 and convert its to out string
+	 * @param SHA					IN		SHA hash to encode
+	 * @param OutString				OUT		String with base32 encoded value
+	 */
+	static void SHAToBase32(const FSHAHash& SHA, FString& OutString);
 
 	static uint64 CalculateDiskSpaceRequirementsWithDeleteDuringInstall(
 		const TArray<FString>& InFilesToConstruct, 
 		int32 InCompletedFileCount, 
+		int64 InProgressFileSize,
 		BuildPatchServices::IBuildManifestSet* InManifestSet, 
 		const TArray<uint64>& InChunkDbSizesAtPosition, 
-		uint64 InTotalChunkDbSize);
+		const FString& InstallDirectory,
+		BuildPatchServices::EInstallMode mode);
 
     /**
 	* Generates the path to the file we're installing by adjusting the path for per file subdirectories. Currently this does not handle
diff --git a/Engine/Source/Runtime/Online/BuildPatchServices/Private/BuildPatchVerifyChunkData.cpp b/Engine/Source/Runtime/Online/BuildPatchServices/Private/BuildPatchVerifyChunkData.cpp
--- a/Engine/Source/Runtime/Online/BuildPatchServices/Private/BuildPatchVerifyChunkData.cpp
+++ b/Engine/Source/Runtime/Online/BuildPatchServices/Private/BuildPatchVerifyChunkData.cpp
@@ -4,7 +4,10 @@
 #include "HAL/FileManager.h"
 #include "Misc/Paths.h"
 #include "Misc/FileHelper.h"
+#include "Common/Crypto.h"
 #include "Common/FileSystem.h"
+#include "Common/StatsCollector.h"
+#include "Installer/Statistics/CloudChunkSourceStat.h"
 #include "Data/ChunkData.h"
 #include "Async/Async.h"
 #include "Generation/ChunkDatabaseWriter.h"
@@ -43,6 +46,7 @@ bool FBuildVerifyChunkData::VerifyChunkData(const FString& SearchPath, const FSt
 
 	TSet<FGuid> BadChunkData;
 	TSet<FGuid> GoodChunkData;
+	TSet<FGuid> SkippedChunkData;
 	TArray<FString> ChunkFiles;
 	TArray<FString> ChunkDbFiles;
 	TArray<FString> ManifestFiles;
@@ -69,10 +73,15 @@ bool FBuildVerifyChunkData::VerifyChunkData(const FString& SearchPath, const FSt
 
 	// Systems.
 	TUniquePtr<IFileSystem> FileSystem(FFileSystemFactory::Create());
-	TUniquePtr<IChunkDataSerialization> ChunkDataSerialization(FChunkDataSerializationFactory::Create(FileSystem.Get()));
+	TUniquePtr<ICrypto> Crypto(FCryptoFactory::Create());
+	TUniquePtr<FStatsCollector> StatsCollector(FStatsCollectorFactory::Create());
+	TUniquePtr<ICloudChunkSourceStat> CloudChunkSourceStat(FCloudChunkSourceStatFactory::Create(StatsCollector.Get()));
+	FChunkDataSerializationConfig ChunkDataSerializationConfig;
+	// UE5 MERGE TODO: ChunkDataSerializationConfig.CloudChunkSourceStat = CloudChunkSourceStat.Get();
+	TUniquePtr<IChunkDataSerialization> ChunkDataSerialization(FChunkDataSerializationFactory::Create(FileSystem.Get(), Crypto.Get(), ChunkDataSerializationConfig));
 
 	// Kick off a verify chunk files.
-	typedef TTuple<FString, bool, FGuid> FChunkFileResult;
+	typedef TTuple<FString, EChunkLoadResult, FGuid> FChunkFileResult;
 	TArray<TFuture<FChunkFileResult>> ChunkFileResults;
 	for (const FString& ChunkFile : ChunkFiles)
 	{
@@ -84,15 +93,15 @@ bool FBuildVerifyChunkData::VerifyChunkData(const FString& SearchPath, const FSt
 			FBuildPatchUtils::GetGUIDFromFilename(ChunkFile, ChunkId);
 			TUniquePtr<IChunkDataAccess> ChunkDataAccess(ChunkDataSerialization->LoadFromFile(ChunkFile, LoadResult));
 			Result.Get<0>() = ChunkFile;
-			Result.Get<1>() = LoadResult == EChunkLoadResult::Success && ChunkDataAccess.IsValid();
+			Result.Get<1>() = LoadResult;
 			Result.Get<2>() = ChunkId;
-			if (Result.Get<1>())
+			if (ChunkDataAccess)
 			{
 				// Ensure we can get GUID from readable chunks.
 				FScopeLockedChunkData ScopeLockedChunkData(ChunkDataAccess.Get());
 				if (ChunkId != ScopeLockedChunkData.GetHeader()->Guid)
 				{
-					Result.Get<1>() = false;
+					Result.Get<1>() = EChunkLoadResult::CorruptHeader;
 				}
 			}
 			return Result;
@@ -121,11 +130,16 @@ bool FBuildVerifyChunkData::VerifyChunkData(const FString& SearchPath, const FSt
 	for (const TFuture<FChunkFileResult>& ChunkFileResult : ChunkFileResults)
 	{
 		FChunkFileResult Result = ChunkFileResult.Get();
-		if (Result.Get<1>())
+		if (Result.Get<1>() == EChunkLoadResult::Success)
 		{
 			UE_LOG(LogVerifyChunkData, Display, TEXT("[%d/%d]: Chunk file good: %s"), ChunkCount++, ChunkNum, *Result.Get<0>());
 			GoodChunkData.Add(Result.Get<2>());
 		}
+		else if (Result.Get<1>() == EChunkLoadResult::MissingSecretKey)
+		{
+			UE_LOG(LogVerifyChunkData, Display, TEXT("[%d/%d]: Chunk file cannot be verified without secret key: %s"), ChunkCount++, ChunkNum, *Result.Get<0>());
+			SkippedChunkData.Add(Result.Get<2>());
+		}
 		else
 		{
 			UE_LOG(LogVerifyChunkData, Error, TEXT("[%d/%d]: Corrupt chunk file: %s"), ChunkCount++, ChunkNum, *Result.Get<0>());
@@ -191,7 +205,6 @@ bool FBuildVerifyChunkData::VerifyChunkData(const FString& SearchPath, const FSt
 		}
 	}
 
-
 	// Collect all the manifest file loads, and see if any are referencing bad data.
 	const int32 ManifestNum = ManifestFileResults.Num();
 	int32 ManifestCount = 0;
@@ -206,8 +219,13 @@ bool FBuildVerifyChunkData::VerifyChunkData(const FString& SearchPath, const FSt
 			{
 				TSet<FGuid> ReferencedData;
 				ManifestFile.Get<0>()->GetDataList(ReferencedData);
+				TSet<FGuid> ReferencedSkippedChunkData = ReferencedData.Intersect(SkippedChunkData);
 				TSet<FGuid> ReferencedBadData = ReferencedData.Intersect(BadChunkData);
 				TSet<FGuid> ReferencedMissingData = ReferencedData.Difference(BadChunkData).Difference(GoodChunkData);
+				if (ReferencedSkippedChunkData.Num() > 0)
+				{
+					UE_LOG(LogVerifyChunkData, Warning, TEXT("[%d/%d]: Skipped data referenced by %s file: %s"), ManifestCount, ManifestNum, *FPaths::GetExtension(ManifestFile.Get<1>()), *ManifestFile.Get<1>());
+				}
 				if (ReferencedBadData.Num() > 0)
 				{
 					UE_LOG(LogVerifyChunkData, Error, TEXT("[%d/%d]: Bad data referenced by %s file: %s"), ManifestCount, ManifestNum, *FPaths::GetExtension(ManifestFile.Get<1>()), *ManifestFile.Get<1>());
diff --git a/Engine/Source/Runtime/Online/BuildPatchServices/Private/Common/ChunkDataSizeProvider.cpp b/Engine/Source/Runtime/Online/BuildPatchServices/Private/Common/ChunkDataSizeProvider.cpp
--- a/Engine/Source/Runtime/Online/BuildPatchServices/Private/Common/ChunkDataSizeProvider.cpp
+++ b/Engine/Source/Runtime/Online/BuildPatchServices/Private/Common/ChunkDataSizeProvider.cpp
@@ -2,6 +2,7 @@
 
 #include "Common/ChunkDataSizeProvider.h"
 #include "Misc/Paths.h"
+#include "Misc/ScopeRWLock.h"
 #include "BuildPatchUtil.h"
 
 namespace BuildPatchServices
@@ -10,6 +11,14 @@ namespace BuildPatchServices
 		: public IChunkDataSizeProvider
 	{
 	public:
+		FChunkDataSizeProvider(const TArray<FBuildPatchAppManifestPtr>& InitialManifests)
+		{
+			for (const FBuildPatchAppManifestPtr& InitialManifest : InitialManifests)
+			{
+				AddManifestData(InitialManifest);
+			}
+		}
+
 		// Begin IDataSizeProvider
 		virtual uint64 GetDownloadSize(const FString& Identifier) const override
 		{
@@ -113,8 +122,8 @@ namespace BuildPatchServices
 		mutable FCriticalSection ManifestsCS;
 	};
 	
-	IChunkDataSizeProvider* FChunkDataSizeProviderFactory::Create()
+	IChunkDataSizeProvider* FChunkDataSizeProviderFactory::Create(const TArray<FBuildPatchAppManifestPtr>& InitialManifests)
 	{
-		return new FChunkDataSizeProvider();
+		return new FChunkDataSizeProvider(InitialManifests);
 	}
 }
diff --git a/Engine/Source/Runtime/Online/BuildPatchServices/Private/Common/ChunkDataSizeProvider.h b/Engine/Source/Runtime/Online/BuildPatchServices/Private/Common/ChunkDataSizeProvider.h
--- a/Engine/Source/Runtime/Online/BuildPatchServices/Private/Common/ChunkDataSizeProvider.h
+++ b/Engine/Source/Runtime/Online/BuildPatchServices/Private/Common/ChunkDataSizeProvider.h
@@ -18,6 +18,6 @@ namespace BuildPatchServices
 	class FChunkDataSizeProviderFactory
 	{
 	public:
-		static IChunkDataSizeProvider* Create();
+		static IChunkDataSizeProvider* Create(const TArray<FBuildPatchAppManifestPtr>& InitialManifests);
 	};
 }
\ No newline at end of file
diff --git a/Engine/Source/Runtime/Online/BuildPatchServices/Private/Common/FileSystem.cpp b/Engine/Source/Runtime/Online/BuildPatchServices/Private/Common/FileSystem.cpp
--- a/Engine/Source/Runtime/Online/BuildPatchServices/Private/Common/FileSystem.cpp
+++ b/Engine/Source/Runtime/Online/BuildPatchServices/Private/Common/FileSystem.cpp
@@ -7,6 +7,9 @@
 #include "Containers/StringConv.h"
 #include "GenericPlatform/GenericPlatformFile.h"
 #include "HAL/FileManager.h"
+#include "HAL/PlatformFileManager.h"
+#include "HAL/FileManagerGeneric.h"
+#include "Misc/OutputDeviceRedirector.h"
 #include "Misc/FileHelper.h"
 
 #if PLATFORM_MAC
@@ -100,7 +103,7 @@ namespace FileSystemHelpers
 }
 // End of region that uses windows types.
 #include "Windows/HideWindowsPlatformTypes.h"
-#elif PLATFORM_MAC
+#elif PLATFORM_MAC || PLATFORM_LINUX
 namespace FileSystemHelpers
 {
 	bool PlatformGetAttributes(const TCHAR* Filename, BuildPatchServices::EAttributeFlags& OutFileAttributes)
@@ -189,14 +192,12 @@ static_assert((uint32)BuildPatchServices::EReadFlags::AllowWrite == (uint32)::EF
 
 namespace BuildPatchServices
 {
-	class FParallelDirectoryEnumerator : public IPlatformFile::FDirectoryVisitor
+	class FDirectoryEnumeratorBase : public IPlatformFile::FDirectoryVisitor
 	{
-	public:
-
-		FParallelDirectoryEnumerator(IPlatformFile& InPlatformFile, const TCHAR* InFileExtension, EAsyncExecution InAsyncExecution)
-			: PlatformFile(InPlatformFile)
+	protected:
+		FDirectoryEnumeratorBase(const TCHAR* InFileExtension, bool bInFollowSymlinks)
+			: bFollowSymlinks(bInFollowSymlinks)
 			, FileExtension(InFileExtension)
-			, AsyncExecution(InAsyncExecution)
 		{
 			if (FileExtension.Len() > 0 && !FileExtension.StartsWith(TEXT(".")))
 			{
@@ -204,46 +205,108 @@ namespace BuildPatchServices
 			}
 		}
 
+		bool IsSymlink(const TCHAR* Filename)
+		{
+			bool bIsSymlink = false;
+			if (!bFollowSymlinks)
+			{
+#if PLATFORM_MAC || PLATFORM_LINUX
+				struct stat StatBuf = {0};
+				if (lstat(TCHAR_TO_UTF8(Filename), &StatBuf) != -1)
+				{
+					bIsSymlink = S_ISLNK(StatBuf.st_mode);
+				}
+#endif
+			}
+			return bIsSymlink;
+		}
+
+		bool MatchesExtension(const TCHAR* Filename) const
+		{
+			if (FileExtension.Len() > 0)
+			{
+				const int32 FileNameLen = FCString::Strlen(Filename);
+				if (FileNameLen < FileExtension.Len() || FCString::Strcmp(&Filename[FileNameLen - FileExtension.Len()], *FileExtension) != 0)
+				{
+					return false;
+				}
+			}
+			return true;
+		}
+
+	private:
+		bool bFollowSymlinks;
+		FString FileExtension;
+	};
+
+	class FDirectoryEnumerator : public FDirectoryEnumeratorBase
+	{
+	public:
+		FDirectoryEnumerator(IPlatformFile& InPlatformFile, const TCHAR* InFileExtension, bool bInFollowSymlinks)
+			: FDirectoryEnumeratorBase(InFileExtension, bInFollowSymlinks)
+			, PlatformFile(InPlatformFile)
+		{  }
+
 		virtual bool Visit(const TCHAR* FilenameOrDirectory, bool bIsDirectory) override
 		{
-			if (bIsDirectory)
+			const bool bIsSymlink = IsSymlink(FilenameOrDirectory);
+			if (bIsDirectory && !bIsSymlink)
 			{
 				FString Directory = FilenameOrDirectory;
-				DirectoryFutures.Enqueue(Async(AsyncExecution, [this, Directory = MoveTemp(Directory)]() { PlatformFile.IterateDirectory(*Directory, *this); }));
+				PlatformFile.IterateDirectory(*Directory, *this);
 			}
 			else if (MatchesExtension(FilenameOrDirectory))
 			{
-				FoundFilesQueue.Enqueue(FilenameOrDirectory);
+				FoundFiles.Add(FilenameOrDirectory);
 			}
 			return true;
 		}
 
 		void GetFiles(TArray<FString>& Results)
 		{
-			TFuture<void> Future;
-			while (DirectoryFutures.Dequeue(Future)) { Future.Wait(); }
-			while (FoundFilesQueue.Dequeue(Results.AddDefaulted_GetRef())) {}
-			Results.Pop(EAllowShrinking::No);
-			Results.Sort();
+			Results = FoundFiles;
 		}
 
 	private:
-		bool MatchesExtension(const TCHAR* Filename) const
+		IPlatformFile& PlatformFile;
+		TArray<FString> FoundFiles;
+	};
+
+	class FParallelDirectoryEnumerator : public FDirectoryEnumeratorBase
+	{
+	public:
+		FParallelDirectoryEnumerator(IPlatformFile& InPlatformFile, const TCHAR* InFileExtension, EAsyncExecution InAsyncExecution, bool bInFollowSymlinks)
+			: FDirectoryEnumeratorBase(InFileExtension, bInFollowSymlinks)
+			, PlatformFile(InPlatformFile)
+			, AsyncExecution(InAsyncExecution)
+		{  }
+
+		virtual bool Visit(const TCHAR* FilenameOrDirectory, bool bIsDirectory) override
 		{
-			if (FileExtension.Len() > 0)
+			const bool bIsSymlink = IsSymlink(FilenameOrDirectory);
+			if (bIsDirectory && !bIsSymlink)
 			{
-				const int32 FileNameLen = FCString::Strlen(Filename);
-				if (FileNameLen < FileExtension.Len() || FCString::Strcmp(&Filename[FileNameLen - FileExtension.Len()], *FileExtension) != 0)
-				{
-					return false;
-				}
+				FString Directory = FilenameOrDirectory;
+				DirectoryFutures.Enqueue(Async(AsyncExecution, [this, Directory = MoveTemp(Directory)]() { PlatformFile.IterateDirectory(*Directory, *this); }));
+			}
+			else if (MatchesExtension(FilenameOrDirectory))
+			{
+				FoundFilesQueue.Enqueue(FilenameOrDirectory);
 			}
 			return true;
 		}
 
+		void GetFiles(TArray<FString>& Results)
+		{
+			TFuture<void> Future;
+			while (DirectoryFutures.Dequeue(Future)) { Future.Wait(); }
+			while (FoundFilesQueue.Dequeue(Results.AddDefaulted_GetRef())) {}
+			Results.Pop(EAllowShrinking::Yes);
+			Results.Sort();
+		}
+
 	private:
 		IPlatformFile& PlatformFile;
-		FString FileExtension;
 		EAsyncExecution AsyncExecution;
 		TQueue<FString, EQueueMode::Mpsc> FoundFilesQueue;
 		TQueue<TFuture<void>, EQueueMode::Mpsc> DirectoryFutures;
@@ -259,23 +322,29 @@ namespace BuildPatchServices
 		// IFileSystem interface begin.
 		virtual bool DirectoryExists(const TCHAR* DirectoryPath) const override;
 		virtual bool MakeDirectory(const TCHAR* DirectoryPath) const override;
+		virtual bool DeleteDirectoryRecursively(const TCHAR* DirectoryPath) const override;
 		virtual bool GetFileSize(const TCHAR* Filename, int64& FileSize) const override;
 		virtual bool GetAttributes(const TCHAR* Filename, EAttributeFlags& Attributes) const override;
 		virtual bool GetTimeStamp(const TCHAR* Path, FDateTime& TimeStamp) const override;
 		virtual bool SetReadOnly(const TCHAR* Filename, bool bIsReadOnly) const override;
 		virtual bool SetCompressed(const TCHAR* Filename, bool bIsCompressed) const override;
 		virtual bool SetExecutable(const TCHAR* Filename, bool bIsExecutable) const override;
+		virtual bool SetFileSize(const TCHAR* Filename, int64 FileSize) const;
 		virtual TUniquePtr<FArchive> CreateFileReader(const TCHAR* Filename, EReadFlags ReadFlags = EReadFlags::None) const override;
+		virtual TUniquePtr<FArchive> CreateFileReaderEx(const TCHAR* Filename, EReadFlags ReadFlags, IFileHandle** OutHandle) const override;
 		virtual TUniquePtr<FArchive> CreateFileWriter(const TCHAR* Filename, EWriteFlags WriteFlags = EWriteFlags::None) const override;
 		virtual bool LoadFileToString(const TCHAR* Filename, FString& Contents) const override;
 		virtual bool SaveStringToFile(const TCHAR* Filename, const FString& Contents) const override;
+		virtual bool LoadFileToArray(const TCHAR* Filename, TArray<uint8>& Contents) const override;
+		virtual bool SaveArrayToFile(const TCHAR* Filename, const TArray<uint8>& Contents) const override;
 		virtual bool DeleteFile(const TCHAR* Filename) const override;
-		virtual bool MoveFile(const TCHAR* FileDest, const TCHAR* FileSource) const override;
+		virtual bool MoveFile(const TCHAR* FileDest, const TCHAR* FileSource, bool bDoNotRetryOrError) const override;
 		virtual bool CopyFile(const TCHAR* FileDest, const TCHAR* FileSource) const override;
 		virtual bool FileExists(const TCHAR* Filename) const override;
 		virtual void FindFiles(TArray<FString>& FoundFiles, const TCHAR* Directory, const TCHAR* FileExtension = nullptr) const override;
-		virtual void FindFilesRecursively(TArray<FString>& FoundFiles, const TCHAR* Directory, const TCHAR* FileExtension = nullptr) const override;
-		virtual void ParallelFindFilesRecursively(TArray<FString>& FoundFiles, const TCHAR* Directory, const TCHAR* FileExtension = nullptr, EAsyncExecution AsyncExecution = EAsyncExecution::ThreadPool) const override;
+		virtual void FindFilesRecursively(TArray<FString>& FoundFiles, const TCHAR* Directory, const TCHAR* FileExtension = nullptr, bool bFollowSymlinks = true) const override;
+		virtual void ParallelFindFilesRecursively(TArray<FString>& FoundFiles, const TCHAR* Directory, const TCHAR* FileExtension = nullptr, EAsyncExecution AsyncExecution = EAsyncExecution::ThreadPool, bool bFollowSymlinks = true) const override;
+		virtual FFileStatData GetStatData(const TCHAR* Filename) const override;
 		virtual int64 GetAllowedBytesToWriteThrottledStorage(const TCHAR* DestinationPath = nullptr) const override;
 		// IFileSystem interface end.
 
@@ -304,6 +373,11 @@ namespace BuildPatchServices
 		return PlatformFile.CreateDirectoryTree(DirectoryPath);
 	}
 
+	bool FFileSystem::DeleteDirectoryRecursively(const TCHAR* DirectoryPath) const
+	{
+		return PlatformFile.DeleteDirectoryRecursively(DirectoryPath);
+	}
+
 	bool FFileSystem::GetFileSize(const TCHAR* Filename, int64& FileSize) const
 	{
 		FileSize = PlatformFile.FileSize(Filename);
@@ -337,11 +411,41 @@ namespace BuildPatchServices
 		return FileSystemHelpers::PlatformSetExecutable(Filename, bIsExecutable);
 	}
 
+	bool FFileSystem::SetFileSize(const TCHAR* Filename, int64 FileSize) const
+	{
+		TUniquePtr<IFileHandle> Handle{ PlatformFile.OpenWrite(Filename, true, true) };
+		if (Handle.IsValid())
+		{
+			return Handle->Truncate(FileSize);
+		}
+		return false;
+	}
+
 	TUniquePtr<FArchive> FFileSystem::CreateFileReader(const TCHAR* Filename, EReadFlags ReadFlags) const
 	{
 		return TUniquePtr<FArchive>(FileManager.CreateFileReader(Filename, static_cast<uint32>(ReadFlags)));
 	}
 
+	TUniquePtr<FArchive> FFileSystem::CreateFileReaderEx(const TCHAR* Filename, EReadFlags ReadFlags, IFileHandle** OutHandle) const
+	{
+		IFileHandle* Handle = FPlatformFileManager::Get().GetPlatformFile().OpenWrite(Filename, true, true);
+		if (OutHandle)
+		{
+			*OutHandle = Handle;
+		}
+
+		if (!Handle)
+		{
+			if (static_cast<uint32>(ReadFlags) & FILEREAD_NoFail)
+			{
+				GLog->Logf(ELogVerbosity::Type::Fatal, TEXT("CreateFileReaderEx: Failed to read file : %s"), Filename);
+			}
+			return nullptr;
+		}
+		Handle->Seek(0); // Manually move the position to the beginning because Append moved it to EOF.
+		return MakeUnique<FArchiveFileReaderGeneric>(Handle, Filename, Handle->Size(), PLATFORM_FILE_READER_BUFFER_SIZE, static_cast<uint32>(ReadFlags));
+	}
+
 	TUniquePtr<FArchive> FFileSystem::CreateFileWriter(const TCHAR* Filename, EWriteFlags WriteFlags) const
 	{
 		return TUniquePtr<FArchive>(FileManager.CreateFileWriter(Filename, static_cast<uint32>(WriteFlags)));
@@ -357,14 +461,29 @@ namespace BuildPatchServices
 		return FFileHelper::SaveStringToFile(Contents, Filename);
 	}
 
+	bool FFileSystem::LoadFileToArray(const TCHAR* Filename, TArray<uint8>& Contents) const
+	{
+		return FFileHelper::LoadFileToArray(Contents, Filename);
+	}
+
+	bool FFileSystem::SaveArrayToFile(const TCHAR* Filename, const TArray<uint8>& Contents) const
+	{
+		return FFileHelper::SaveArrayToFile(Contents, Filename);
+	}
+
 	bool FFileSystem::DeleteFile(const TCHAR* Filename) const
 	{
-		return FileManager.Delete(Filename, false, true, true);
+		if (FileExists(Filename))
+		{
+			SetReadOnly(Filename, false);
+			return PlatformFile.DeleteFile(Filename);
+		}
+		return true;
 	}
 
-	bool FFileSystem::MoveFile(const TCHAR* FileDest, const TCHAR* FileSource) const
+	bool FFileSystem::MoveFile(const TCHAR* FileDest, const TCHAR* FileSource, bool bDoNotRetryOrError) const
 	{
-		return FileManager.Move(FileDest, FileSource, true, true, true, true);
+		return FileManager.Move(FileDest, FileSource, true, true, true, bDoNotRetryOrError);
 	}
 
 	bool FFileSystem::CopyFile(const TCHAR* FileDest, const TCHAR* FileSource) const
@@ -374,7 +493,16 @@ namespace BuildPatchServices
 
 	bool FFileSystem::FileExists(const TCHAR* Filename) const
 	{
+#if PLATFORM_MAC || PLATFORM_LINUX
+		struct stat StatBuf = {0};
+		if (lstat(TCHAR_TO_UTF8(Filename), &StatBuf) != -1)
+		{
+			return S_ISREG(StatBuf.st_mode) || S_ISLNK(StatBuf.st_mode);
+		}
+		return false;
+#else
 		return FileManager.FileExists(Filename);
+#endif
 	}
 
 	void FFileSystem::FindFiles(TArray<FString>& FoundFiles, const TCHAR* Directory, const TCHAR* FileExtension) const
@@ -382,18 +510,25 @@ namespace BuildPatchServices
 		PlatformFile.FindFiles(FoundFiles, Directory, FileExtension);
 	}
 
-	void FFileSystem::FindFilesRecursively(TArray<FString>& FoundFiles, const TCHAR* Directory, const TCHAR* FileExtension) const
+	void FFileSystem::FindFilesRecursively(TArray<FString>& FoundFiles, const TCHAR* Directory, const TCHAR* FileExtension, bool bFollowSymlinks) const
 	{
-		PlatformFile.FindFilesRecursively(FoundFiles, Directory, FileExtension);
+		FDirectoryEnumerator DirectoryEnumerator(PlatformFile, FileExtension, bFollowSymlinks);
+		PlatformFile.IterateDirectory(Directory, DirectoryEnumerator);
+		DirectoryEnumerator.GetFiles(FoundFiles);
 	}
 
-	void FFileSystem::ParallelFindFilesRecursively(TArray<FString>& FoundFiles, const TCHAR* Directory, const TCHAR* FileExtension, EAsyncExecution AsyncExecution) const
+	void FFileSystem::ParallelFindFilesRecursively(TArray<FString>& FoundFiles, const TCHAR* Directory, const TCHAR* FileExtension, EAsyncExecution AsyncExecution, bool bFollowSymlinks) const
 	{
-		FParallelDirectoryEnumerator DirectoryEnumerator(PlatformFile, FileExtension, AsyncExecution);
+		FParallelDirectoryEnumerator DirectoryEnumerator(PlatformFile, FileExtension, AsyncExecution, bFollowSymlinks);
 		PlatformFile.IterateDirectory(Directory, DirectoryEnumerator);
 		DirectoryEnumerator.GetFiles(FoundFiles);
 	}
 
+	FFileStatData FFileSystem::GetStatData(const TCHAR* FilenameOrDirectory) const
+	{
+		return PlatformFile.GetStatData(FilenameOrDirectory);
+	}
+
 	int64 FFileSystem::GetAllowedBytesToWriteThrottledStorage(const TCHAR* DestinationPath) const
 	{
 		return PlatformFile.GetAllowedBytesToWriteThrottledStorage(DestinationPath);
diff --git a/Engine/Source/Runtime/Online/BuildPatchServices/Private/Common/FileSystem.h b/Engine/Source/Runtime/Online/BuildPatchServices/Private/Common/FileSystem.h
--- a/Engine/Source/Runtime/Online/BuildPatchServices/Private/Common/FileSystem.h
+++ b/Engine/Source/Runtime/Online/BuildPatchServices/Private/Common/FileSystem.h
@@ -2,6 +2,7 @@
 #pragma once
 
 #include "Async/Async.h"
+#include "GenericPlatform/GenericPlatformFile.h"
 #include "Misc/EnumClassFlags.h"
 #include "Templates/UniquePtr.h"
 
@@ -66,6 +67,13 @@ namespace BuildPatchServices
 		 */
 		virtual bool MakeDirectory(const TCHAR* DirectoryPath) const = 0;
 
+		/**
+		 * Delete a directory path including all containing files and folders.
+		 * @param DirectoryPath     The directory path to delete.
+		 * @return true if successful.
+		 */
+		virtual bool DeleteDirectoryRecursively(const TCHAR* DirectoryPath) const = 0;
+
 		/**
 		 * Get the size of a file.
 		 * @param Filename          The filename for the request.
@@ -114,6 +122,14 @@ namespace BuildPatchServices
 		 */
 		virtual bool SetExecutable(const TCHAR* Filename, bool bIsExecutable) const = 0;
 
+		/**
+		 * Sets the size of a file, can increase or decrease, and will also create a file that does not exist..
+		 * @param Filename          The filename for the request.
+		 * @param FileSize          The desired size in bytes.
+		 * @return true if successful.
+		 */
+		virtual bool SetFileSize(const TCHAR* Filename, int64 FileSize) const = 0;
+
 		/**
 		 * Creates file reader archive.
 		 * @param Filename          The filename for the request.
@@ -122,6 +138,15 @@ namespace BuildPatchServices
 		 */
 		virtual TUniquePtr<FArchive> CreateFileReader(const TCHAR* Filename, EReadFlags ReadFlags = EReadFlags::None) const = 0;
 
+		/**
+		 * Creates file reader archive.
+		 * @param Filename          The filename for the request.
+		 * @param ReadFlags         The file open flags.
+		 * @param OutHandle			Pointer to the file handler used by the archive. The returned FArchive owns the instance of OutHandle and controls its lifetime.
+		 * @return unique pointer to the created archive, invalid if failed to open the file.
+		 */
+		virtual TUniquePtr<FArchive> CreateFileReaderEx(const TCHAR* Filename, EReadFlags ReadFlags, IFileHandle** OutHandle) const = 0;
+
 		/**
 		 * Creates file writer archive.
 		 * @param Filename          The filename for the request.
@@ -146,6 +171,22 @@ namespace BuildPatchServices
 		 */
 		virtual bool SaveStringToFile(const TCHAR* Filename, const FString& Contents) const = 0;
 
+		/**
+		 * Load a file on disk into a byte array.
+		 * @param Filename          The filename for the request.
+		 * @param Contents          The contents of the file if successful.
+		 * @return true if successful.
+		 */
+		virtual bool LoadFileToArray(const TCHAR* Filename, TArray<uint8>& Contents) const = 0;
+
+		/**
+		 * Save a byte array to file on disk.
+		 * @param Filename          The filename for the request.
+		 * @param Contents          The string to save.
+		 * @return true if successful.
+		 */
+		virtual bool SaveArrayToFile(const TCHAR* Filename, const TArray<uint8>& Contents) const = 0;
+
 		/**
 		 * Delete a file.
 		 * @param Filename          The file to delete.
@@ -157,9 +198,10 @@ namespace BuildPatchServices
 		 * Move or rename a file.
 		 * @param FileDest          The destination file path.
 		 * @param FileSource        The source file path.
+		 * @param bDoNotRetryOrError Do not try to repeat failed move operation. 
 		 * @return true if the file was moved/renamed successfully.
 		 */
-		virtual bool MoveFile(const TCHAR* FileDest, const TCHAR* FileSource) const = 0;
+		virtual bool MoveFile(const TCHAR* FileDest, const TCHAR* FileSource, bool bDoNotRetryOrError = true) const = 0;
 
 		/**
 		 * Copy a file.
@@ -189,8 +231,9 @@ namespace BuildPatchServices
 		 * @param FoundFiles        Receives the files that matched the optional FileExtension filter, or all files if none was specified.
 		 * @param Directory         The directory to iterate the contents of. This function explores subdirectories.
 		 * @param FileExtension     The extension to filter by, in the form of TEXT(".ext"). If null or empty string, all files are found.
+		 * @param bFollowSymlinks   Whether symlinks to directories should be followed and enumerated, or added as a file.
 		 */
-		virtual void FindFilesRecursively(TArray<FString>& FoundFiles, const TCHAR* Directory, const TCHAR* FileExtension = nullptr) const = 0;
+		virtual void FindFilesRecursively(TArray<FString>& FoundFiles, const TCHAR* Directory, const TCHAR* FileExtension = nullptr, bool bFollowSymlinks = true) const = 0;
 
 		/**
 		 * Finds all the files within the given directory tree, with optional file extension filter. This version places each directory found
@@ -199,8 +242,16 @@ namespace BuildPatchServices
 		 * @param Directory         The directory to iterate the contents of. This function explores subdirectories.
 		 * @param FileExtension     The extension to filter by, in the form of TEXT(".ext"). If null or empty string, all files are found.
 		 * @param AsyncExecution    The threaded execution method to use, i.e. on Task Graph, Thread Pool, or a new thread per directory.
+		 * @param bFollowSymlinks   Whether symlinks to directories should be followed and enumerated, or added as a file.
+		 */
+		virtual void ParallelFindFilesRecursively(TArray<FString>& FoundFiles, const TCHAR* Directory, const TCHAR* FileExtension = nullptr, EAsyncExecution AsyncExecution = EAsyncExecution::ThreadPool, bool bFollowSymlinks = true) const = 0;
+
+		/**
+		 * Return the stat data for the given file or directory. Check the FFileStatData::bIsValid member before using the returned data.
+		 * @param Filename          The filename or directory for the request.
+		 * @return resulted stat data of the requested file or dir.
 		 */
-		virtual void ParallelFindFilesRecursively(TArray<FString>& FoundFiles, const TCHAR* Directory, const TCHAR* FileExtension = nullptr, EAsyncExecution AsyncExecution = EAsyncExecution::ThreadPool) const = 0;
+		virtual FFileStatData GetStatData(const TCHAR* FilenameOrDirectory) const = 0;
 
 		/**
 		 * Returns the number of bytes that are currently allowed to be written to throttled write storage (if the platform
diff --git a/Engine/Source/Runtime/Online/BuildPatchServices/Private/Common/SpeedRecorder.cpp b/Engine/Source/Runtime/Online/BuildPatchServices/Private/Common/SpeedRecorder.cpp
--- a/Engine/Source/Runtime/Online/BuildPatchServices/Private/Common/SpeedRecorder.cpp
+++ b/Engine/Source/Runtime/Online/BuildPatchServices/Private/Common/SpeedRecorder.cpp
@@ -94,7 +94,7 @@ namespace BuildPatchServices
 
 		// ISpeedRecorder interface begin.
 		virtual void AddRecord(const FRecord& Record);
-		virtual double GetAverageSpeed(float Seconds) const;
+		virtual double GetAverageSpeed(float Seconds, bool bIncludeInactiveTime) const;
 		virtual double GetPeakSpeed() const;
 		// ISpeedRecorder interface end.
 
@@ -137,7 +137,7 @@ namespace BuildPatchServices
 		RecordsQueue.Enqueue(Record);
 	}
 
-	double FSpeedRecorder::GetAverageSpeed(float Seconds) const
+	double FSpeedRecorder::GetAverageSpeed(float Seconds, bool bIncludeInactiveTime) const
 	{
 		checkSlow(IsInGameThread());
 		// Fill Temp with correct range data.
@@ -157,7 +157,8 @@ namespace BuildPatchServices
 			}
 			TotalSize += Record.Size;
 		}
-		return TotalCycles > 0 ? TotalSize / FStatsCollector::CyclesToSeconds(TotalCycles) : 0.0;
+		double SecondsDivisor = bIncludeInactiveTime ? Seconds : FStatsCollector::CyclesToSeconds(TotalCycles);
+		return TotalCycles > 0 ? TotalSize / SecondsDivisor : 0.0;
 	}
 
 	double FSpeedRecorder::GetPeakSpeed() const
diff --git a/Engine/Source/Runtime/Online/BuildPatchServices/Private/Common/SpeedRecorder.h b/Engine/Source/Runtime/Online/BuildPatchServices/Private/Common/SpeedRecorder.h
--- a/Engine/Source/Runtime/Online/BuildPatchServices/Private/Common/SpeedRecorder.h
+++ b/Engine/Source/Runtime/Online/BuildPatchServices/Private/Common/SpeedRecorder.h
@@ -37,9 +37,10 @@ namespace BuildPatchServices
 		/**
 		 * Get the current average speed achieved over the last X seconds.
 		 * @param Seconds   The time in seconds to take the reading over. Use TNumericLimits<float>::Max() to get the average over all readings.
+		 * @param bIncludeInactiveTime   If True speed will be averaged over the entire Seconds period. Default is to only average over periods of active download within the Seconds period.
 		 * @return the average speed over the past given seconds.
 		 */
-		virtual double GetAverageSpeed(float Seconds) const = 0;
+		virtual double GetAverageSpeed(float Seconds, bool bIncludeInactiveTime = false) const = 0;
 
 		/**
 		 * Get the peak speed achieved so far.
diff --git a/Engine/Source/Runtime/Online/BuildPatchServices/Private/Common/StatsCollector.cpp b/Engine/Source/Runtime/Online/BuildPatchServices/Private/Common/StatsCollector.cpp
--- a/Engine/Source/Runtime/Online/BuildPatchServices/Private/Common/StatsCollector.cpp
+++ b/Engine/Source/Runtime/Online/BuildPatchServices/Private/Common/StatsCollector.cpp
@@ -77,7 +77,7 @@ namespace BuildPatchServices
 		virtual ~FStatsCollectorImpl();
 
 		virtual volatile FAtomicValue* CreateStat(const FString& Name, EStatFormat Type, FAtomicValue InitialValue = 0) override;
-		virtual void LogStats(float TimeBetweenLogs = 0.0f) override;
+		virtual bool LogStats(float TimeBetweenLogs = 0.0f) override;
 
 	private:
 		FCriticalSection DataCS;
@@ -130,7 +130,7 @@ namespace BuildPatchServices
 		return AddedStats[Name];
 	}
 
-	void FStatsCollectorImpl::LogStats(float TimeBetweenLogs)
+	bool FStatsCollectorImpl::LogStats(float TimeBetweenLogs)
 	{
 		const uint64 Cycles = FStatsCollector::GetCycles();
 		if (FStatsCollector::CyclesToSeconds(Cycles - LastLogged) >= TimeBetweenLogs)
@@ -138,7 +138,7 @@ namespace BuildPatchServices
 			LastLogged = Cycles;
 			FScopeLock ScopeLock(&DataCS);
 			GLog->Log(TEXT("/-------- FStatsCollector Log ---------------------"));
-			for (auto& Stat : AddedStats)
+			for (const auto& Stat : AddedStats)
 			{
 				FAtomicValue* NameLookup = (FAtomicValue*)Stat.Value;
 				switch (AtomicFormatMap[NameLookup])
@@ -161,7 +161,9 @@ namespace BuildPatchServices
 				}
 			}
 			GLog->Log(TEXT("\\--------------------------------------------------"));
+			return true;
 		}
+		return false;
 	}
 
 	FStatsCollector* FStatsCollectorFactory::Create()
diff --git a/Engine/Source/Runtime/Online/BuildPatchServices/Private/Common/StatsCollector.h b/Engine/Source/Runtime/Online/BuildPatchServices/Private/Common/StatsCollector.h
--- a/Engine/Source/Runtime/Online/BuildPatchServices/Private/Common/StatsCollector.h
+++ b/Engine/Source/Runtime/Online/BuildPatchServices/Private/Common/StatsCollector.h
@@ -27,7 +27,7 @@ namespace BuildPatchServices
 	public:
 		virtual ~FStatsCollector() = default;
 		virtual volatile FAtomicValue* CreateStat(const FString& Name, EStatFormat Type, FAtomicValue InitialValue = 0) = 0;
-		virtual void LogStats(float TimeBetweenLogs = 0.0f) = 0;
+		virtual bool LogStats(float TimeBetweenLogs = 0.0f) = 0;
 
 	public:
 		static uint64 GetCycles();
@@ -48,6 +48,23 @@ namespace BuildPatchServices
 		static FStatsCollector* Create();
 	};
 
+	class FStatsScopedCounter
+	{
+	public:
+		FStatsScopedCounter(volatile FStatsCollector::FAtomicValue* InStat)
+			: Stat(InStat)
+		{
+			FStatsCollector::Accumulate(Stat, 1);
+		}
+		~FStatsScopedCounter()
+		{
+			FStatsCollector::Accumulate(Stat, -1);
+		}
+
+	private:
+		volatile FStatsCollector::FAtomicValue* Stat;
+	};
+
 	class FStatsScopedTimer
 	{
 	public:
diff --git a/Engine/Source/Runtime/Online/BuildPatchServices/Private/Compactify/PatchDataCompactifier.cpp b/Engine/Source/Runtime/Online/BuildPatchServices/Private/Compactify/PatchDataCompactifier.cpp
--- a/Engine/Source/Runtime/Online/BuildPatchServices/Private/Compactify/PatchDataCompactifier.cpp
+++ b/Engine/Source/Runtime/Online/BuildPatchServices/Private/Compactify/PatchDataCompactifier.cpp
@@ -2,64 +2,141 @@
 
 #include "Compactify/PatchDataCompactifier.h"
 
+#include "Algo/Accumulate.h"
 #include "Algo/Partition.h"
 #include "Algo/Transform.h"
 #include "Async/Async.h"
-#include "Logging/LogMacros.h"
-
 #include "Common/FileSystem.h"
+#include "Common/StatsCollector.h"
+#include "Core/AsyncHelpers.h"
 #include "Core/ProcessTimer.h"
 #include "Enumeration/PatchDataEnumeration.h"
+#include "Logging/LogMacros.h"
+#include "Misc/ConfigCacheIni.h"
+#include "Misc/OutputDeviceRedirector.h"
 
-DECLARE_LOG_CATEGORY_CLASS(LogDataCompactifier, Log, All);
+DECLARE_LOG_CATEGORY_EXTERN(LogDataCompactifier, Log, All);
+DEFINE_LOG_CATEGORY(LogDataCompactifier);
 
 namespace BuildPatchServices
 {
+	constexpr int MaxManifestHandlingThreads = 150;
+	constexpr int MaxFileHandlingThreads = 500;
+
+	struct FReferencedFileData
+	{
+	public:
+		FString Name;
+		int64 Size = 0;
+
+		bool operator==(const FString& Rhs) const
+		{
+			return Name == Rhs;
+		}
+		bool operator==(const FReferencedFileData& Rhs) const
+		{
+			return Name == Rhs.Name;
+		}
+	};
+
+	FORCEINLINE uint32 GetTypeHash(const FReferencedFileData& Data)
+	{
+		return GetTypeHash(Data.Name);
+	}
+
 	struct FManifestResult
 	{
 	public:
-		FManifestResult()
-			: bSuccess(false)
-			, FileSize(0)
-		{ }
+		FManifestResult() = default;
 
-		FManifestResult(bool bInSuccess, int64 InFileSize, TSet<FString>&& InReferences)
+		FManifestResult(bool bInSuccess, int64 InFileSize, TSet<FReferencedFileData>&& InReferences)
 			: bSuccess(bInSuccess)
 			, FileSize(InFileSize)
 			, References(MoveTemp(InReferences))
 		{ }
 
 	public:
-		bool bSuccess;
-		int64 FileSize;
-		TSet<FString> References;
+		bool bSuccess = false;
+		int64 FileSize = 0;
+		TSet<FReferencedFileData> References;
+	};
+
+	struct FStatsTimeGuard
+	{
+	public:
+		FStatsTimeGuard(FStatsCollector& StatsCollector, const TCHAR* Type)
+		{
+			TotalTime = StatsCollector.CreateStat(Type, EStatFormat::Timer);
+		}
+
+		~FStatsTimeGuard()
+		{
+			Refresh();
+		}
+
+		FStatsTimeGuard& Refresh()
+		{
+			FStatsCollector::Set(TotalTime, FStatsCollector::GetCycles() - StartTime);
+			return *this;
+		}
+
+		FString GetAsString()
+		{
+			return FPlatformTime::PrettyTime(FStatsCollector::CyclesToSeconds(*TotalTime));
+		}
+
+	private:
+		volatile int64* TotalTime;
+		const uint64 StartTime = FStatsCollector::GetCycles();
 	};
 
+	struct FProcessingStats
+	{
+	public:
+		FProcessingStats(FStatsCollector& StatsCollector, const TCHAR* Type)
+		{
+			Num = StatsCollector.CreateStat(FString::Printf(TEXT("Data Compactifier: %s Num"), Type), EStatFormat::Value);
+			Bytes = StatsCollector.CreateStat(FString::Printf(TEXT("Data Compactifier: %s Bytes"), Type), EStatFormat::DataSize);
+		}
+
+		void Accumulate(uint64 ProcessedBytes)
+		{
+			Accumulate(1, ProcessedBytes);
+		}
+
+		void Accumulate(uint64 FilesNum, uint64 ProcessedBytes)
+		{
+			FStatsCollector::Accumulate(Num, FilesNum);
+			FStatsCollector::Accumulate(Bytes, ProcessedBytes);
+		}
+
+		void Accumulate(const FProcessingStats& Stats)
+		{
+			Accumulate(*Stats.Num, *Stats.Bytes);
+		}
+
+	public:
+		volatile FStatsCollector::FAtomicValue* Num;
+		volatile FStatsCollector::FAtomicValue* Bytes;
+	};	
+
 	struct FDataFileResult
 	{
 	public:
-		FDataFileResult()
-			: bIsReferenced(false)
-			, bIsOldEnough(false)
-			, bIsRecognisedFileType(false)
-			, bShouldDelete(false)
-			, FileSize(0)
-		{ }
+		FDataFileResult() = default;
 
-		FDataFileResult(bool bInIsReferenced, bool bInIsOldEnough, bool bInIsRecognisedFileType, bool bInShouldDelete, int64 InFileSize)
-			: bIsReferenced(bInIsReferenced)
-			, bIsOldEnough(bInIsOldEnough)
+		FDataFileResult(bool bInIsOldEnough, bool bInIsRecognisedFileType, bool bInShouldDelete, int64 InFileSize)
+			: bIsOldEnough(bInIsOldEnough)
 			, bIsRecognisedFileType(bInIsRecognisedFileType)
 			, bShouldDelete(bInShouldDelete)
 			, FileSize(InFileSize)
 		{ }
 
 	public:
-		bool bIsReferenced;
-		bool bIsOldEnough;
-		bool bIsRecognisedFileType;
-		bool bShouldDelete;
-		int64 FileSize;
+		bool bIsOldEnough = false;
+		bool bIsRecognisedFileType = false;
+		bool bShouldDelete = false;
+		int64 FileSize = 0;
 	};
 
 	class FPatchDataCompactifier
@@ -74,19 +151,23 @@ namespace BuildPatchServices
 		// IPatchDataCompactifier interface end.
 
 	private:
+		void LogStats(float TimeBetweenLogs = 0.f);
+		TArray<FString> FindAllFiles();
+		bool FindAllReferencedFiles(float StatsLoggerTimeSeconds, int32 FirstNonManifest, const TArray<FString>& AllFiles, FProcessingStats& ProcessedFilesStat, TSet<FReferencedFileData>& ReferencedFileSet);
 		void DeleteFile(const FString& FilePath) const;
 		bool IsPatchData(const FString& FilePath) const;
-		FString BuildSizeString(uint64 Size) const;
 
 	private:
 		const FCompactifyConfiguration Configuration;
-		FNumberFormattingOptions SizeFormattingOptions;
 		TUniquePtr<IFileSystem> FileSystem;
+		TUniquePtr<FStatsCollector> StatsCollector{FStatsCollectorFactory::Create()};
+		volatile int64* StatTotalTime = StatsCollector->CreateStat(TEXT("Data Compactifier: Total Time"), EStatFormat::Timer);
+		uint64 StartCompactifierTime = FStatsCollector::GetCycles();
+		volatile FStatsCollector::FAtomicValue* AllFilesProgress = StatsCollector->CreateStat(TEXT("Data Compactifier: All Files Progress"), EStatFormat::Percentage);
 	};
 
 	FPatchDataCompactifier::FPatchDataCompactifier(const FCompactifyConfiguration& InConfiguration)
 		: Configuration(InConfiguration)
-		, SizeFormattingOptions(FNumberFormattingOptions().SetMaximumFractionalDigits(3).SetMinimumFractionalDigits(3))
 		, FileSystem(FFileSystemFactory::Create())
 	{
 	}
@@ -95,8 +176,121 @@ namespace BuildPatchServices
 	{
 	}
 
+	void FPatchDataCompactifier::LogStats(float TimeBetweenLogs)
+	{
+		FStatsCollector::Set(StatTotalTime, FStatsCollector::GetCycles() - StartCompactifierTime);
+		if (StatsCollector->LogStats(TimeBetweenLogs))
+		{
+			GLog->Flush();
+		}
+	}
+
+	TArray<FString> FPatchDataCompactifier::FindAllFiles()
+	{
+		FStatsTimeGuard FindFilesTotalTimeStat(*StatsCollector, TEXT("Data Compactifier: Find All Files Time"));
+		TArray<FString> AllFiles;
+		FileSystem->ParallelFindFilesRecursively(AllFiles, *Configuration.CloudDirectory, nullptr, EAsyncExecution::Thread);
+		UE_LOG(LogDataCompactifier, Log, TEXT("Found %d files in %s"), AllFiles.Num(), *FindFilesTotalTimeStat.Refresh().GetAsString());
+		return AllFiles;
+	}
+
+	bool FPatchDataCompactifier::FindAllReferencedFiles(float StatsLoggerTimeSeconds, int32 FirstNonManifest, const TArray<FString>& AllFiles, FProcessingStats &ProcessedFilesStat, TSet<FReferencedFileData>& ReferencedFileSet)
+	{
+		// Track some statistics.
+		FProcessingStats ManifestFilesStat(*StatsCollector, TEXT("Manifest Files"));
+		FProcessingStats ReferencedFilesStat(*StatsCollector, TEXT("Referenced Files"));
+
+		FStatsTimeGuard FindAllReferencedFilesTimeStat(*StatsCollector, TEXT("Data Compactifier: Handle Manifests Time"));
+
+		bool bSuccess = true;
+		FTaskScheduler<FManifestResult, MaxManifestHandlingThreads> ManifestTasks;
+		for (int32 FileIdx = 0; FileIdx < FirstNonManifest; ++FileIdx)
+		{
+			const FString& Filename = AllFiles[FileIdx];
+			auto TaskFunc = [&]()
+			{
+				bool bEnumerateSuccess = true;
+				TSet<FReferencedFileData> ManifestReferences;
+				FFileStatData FileStatData = FileSystem->GetStatData(*Filename);
+				const bool bGotFileSize = FileStatData.FileSize >= 0;
+				const bool bFileStillExists = FileStatData.bIsValid;
+				if (bGotFileSize)
+				{
+					BuildPatchServices::FPatchDataEnumerationConfiguration EnumerationConfig;
+					EnumerationConfig.InputFile = Filename;
+					EnumerationConfig.bIncludeSizes = true;
+					TUniquePtr<IPatchDataEnumeration> PatchDataEnumeration(FPatchDataEnumerationFactory::Create(EnumerationConfig));
+					TArray<FString> ManifestReferencesArray;
+					bEnumerateSuccess = PatchDataEnumeration->Run(ManifestReferencesArray);
+					if (bEnumerateSuccess)
+					{
+						ManifestReferences.Reserve(ManifestReferencesArray.Num());
+						for (FString& Elem : ManifestReferencesArray)
+						{
+							FReferencedFileData FileData;
+							int32 BytesPos = Elem.Find("\t");
+							if (BytesPos > 0)
+							{
+								FileData.Size = FCString::Atoi(*Elem + BytesPos + 1); // skip \t
+								Elem.LeftInline(BytesPos);
+							}
+							else
+							{
+								UE_LOG(LogDataCompactifier, Warning, TEXT("Could not find a size position of %s. Assuming 0 bytes."), *Elem);
+							}
+							FileData.Name = Configuration.CloudDirectory / Elem;
+							ManifestReferences.Add(MoveTemp(FileData));
+						}
+					}
+				}
+				else if (bFileStillExists)
+				{
+					UE_LOG(LogDataCompactifier, Warning, TEXT("Could not determine size of %s. Assuming 0 bytes."), *Filename);
+					FileStatData.FileSize = 0;
+				}
+				else
+				{
+					UE_LOG(LogDataCompactifier, Log, TEXT("File removed since enumeration %s. Using 0 bytes."), *Filename);
+					FileStatData.FileSize = 0;
+				}
+				return FManifestResult(bEnumerateSuccess, FileStatData.FileSize, MoveTemp(ManifestReferences));
+			};
+			auto CompleteFunc = [&](FManifestResult Result)
+			{
+				ManifestFilesStat.Accumulate(Result.FileSize);
+				if (!Result.bSuccess)
+				{
+					bSuccess = false;
+					UE_LOG(LogDataCompactifier, Error, TEXT("Failed to extract references from %s."), *Filename);
+				}
+				else
+				{
+					const int32 NumReferences = Result.References.Num();
+					ReferencedFileSet.Append(MoveTemp(Result.References));
+					UE_LOG(LogDataCompactifier, Display, TEXT("Extracted %d references from %s. Unioning with existing files, new count of %d."), NumReferences, *Filename, ReferencedFileSet.Num());
+				}
+			};
+
+			ManifestTasks.ScheduleTask(MoveTemp(TaskFunc), MoveTemp(CompleteFunc));
+			FindAllReferencedFilesTimeStat.Refresh();
+			LogStats(StatsLoggerTimeSeconds);
+		}
+		ManifestTasks.WaitAll();
+
+		uint64 FilesSize = Algo::Accumulate(ReferencedFileSet, uint64(0), [](uint64 Result, const FReferencedFileData& FileData) { return Result + FileData.Size; });
+		ReferencedFilesStat.Accumulate(ReferencedFileSet.Num(), FilesSize);
+		ProcessedFilesStat.Accumulate(ReferencedFilesStat);
+		FStatsCollector::SetAsPercentage(AllFilesProgress, (double)FirstNonManifest / (double)AllFiles.Num());
+
+		LogStats();
+		UE_LOG(LogDataCompactifier, Display, TEXT("Manifests handling completed in %s. Found %u references in manifests."), *FindAllReferencedFilesTimeStat.Refresh().GetAsString(), ReferencedFileSet.Num());
+
+		return bSuccess;
+	}
+
 	bool FPatchDataCompactifier::Run()
 	{
+		StartCompactifierTime = FStatsCollector::GetCycles();
 		bool bSuccess = true;
 
 		// We output filenames deleted if requested.
@@ -113,28 +307,20 @@ namespace BuildPatchServices
 
 		if (bSuccess)
 		{
-			// Track some statistics.
-			uint32 FilesProcessed = 0;
-			uint32 FilesSkipped = 0;
-			uint32 NonPatchFilesProcessed = 0;
-			uint32 FilesDeleted = 0;
-			uint64 BytesProcessed = 0;
-			uint64 BytesSkipped = 0;
-			uint64 NonPatchBytesProcessed = 0;
-			uint64 BytesDeleted = 0;
-
 			// We'll work out the date of the oldest unreferenced file we'll keep
 			FDateTime Cutoff = FDateTime::UtcNow() - FTimespan::FromDays(Configuration.DataAgeThreshold);
 
 			// List all files first, and then we'll work with the list.
-			TArray<FString> AllFiles;
-			const bool bFindFiles = true;
-			const bool bFindDirectories = false;
-			FProcessTimer FindFilesTimer;
-			FindFilesTimer.Start();
-			FileSystem->ParallelFindFilesRecursively(AllFiles, *Configuration.CloudDirectory, nullptr, EAsyncExecution::Thread);
-			FindFilesTimer.Stop();
-			UE_LOG(LogDataCompactifier, Log, TEXT("Found %d files in %s."), AllFiles.Num(), *FPlatformTime::PrettyTime(FindFilesTimer.GetSeconds()));
+			TArray<FString> AllFiles = FindAllFiles();
+			if (AllFiles.Num() == 0)
+			{
+				UE_LOG(LogDataCompactifier, Log, TEXT("Could not find any file in the cloud directory %s."), *Configuration.CloudDirectory);
+				return true;
+			}
+
+			volatile FStatsCollector::FAtomicValue* AllFilesNum = StatsCollector->CreateStat(TEXT("Data Compactifier: All Files Num"), EStatFormat::Value);
+			FStatsCollector::Accumulate(AllFilesNum, AllFiles.Num());
+			FStatsCollector::SetAsPercentage(AllFilesProgress, 0);
 
 			// Filter for manifest files.
 			int32 FirstNonManifest = Algo::Partition(AllFiles.GetData(), AllFiles.Num(), [&](const FString& Elem) { return Elem.EndsWith(TEXT(".manifest"), ESearchCase::IgnoreCase); });
@@ -145,152 +331,92 @@ namespace BuildPatchServices
 				UE_LOG(LogDataCompactifier, Log, TEXT("Could not find any manifest files. Proceeding to delete all mature data."));
 			}
 
-			// Handle all manifests first.
-			TArray<TFuture<FManifestResult>> ManifestFutures;
-			for (int32 FileIdx = 0; FileIdx < FirstNonManifest; ++FileIdx)
-			{
-				const FString& Filename = AllFiles[FileIdx];
-				ManifestFutures.Add(Async(EAsyncExecution::ThreadPool, [this, &Filename]()
-				{
-					bool bEnumerateSuccess = true;
-					TSet<FString> ManifestReferences;
-					int64 FileSize = 0;
-					const bool bGotFileSize = FileSystem->GetFileSize(*Filename, FileSize);
-					const bool bFileStillExists = FileSystem->FileExists(*Filename);
-					if (bGotFileSize)
-					{
-						BuildPatchServices::FPatchDataEnumerationConfiguration EnumerationConfig;
-						EnumerationConfig.InputFile = Filename;
-						TUniquePtr<IPatchDataEnumeration> PatchDataEnumeration(FPatchDataEnumerationFactory::Create(EnumerationConfig));
-						TArray<FString> ManifestReferencesArray;
-						bEnumerateSuccess = PatchDataEnumeration->Run(ManifestReferencesArray);
-						if (bEnumerateSuccess)
-						{
-							Algo::Transform(ManifestReferencesArray, ManifestReferences, [&](const FString& Elem) { return Configuration.CloudDirectory / Elem; });
-						}
-					}
-					else if (bFileStillExists)
-					{
-						UE_LOG(LogDataCompactifier, Warning, TEXT("Could not determine size of %s. Assuming 0 bytes."), *Filename);
-						FileSize = 0;
-					}
-					else
-					{
-						UE_LOG(LogDataCompactifier, Log, TEXT("File removed since enumeration %s. Using 0 bytes."), *Filename);
-						FileSize = 0;
-					}
-					return FManifestResult(bEnumerateSuccess, FileSize, MoveTemp(ManifestReferences));
-				}));
-			}
-			TSet<FString> ReferencedFileSet;
-			int32 ManifestFileIdx = 0;
-			for (TFuture<FManifestResult>& ManifestFuture : ManifestFutures)
-			{
-				// Grab result.
-				FManifestResult ManifestResult = ManifestFuture.Get();
-				const FString& ManifestFilename = AllFiles[ManifestFileIdx];
-				// Process result.
-				++FilesProcessed;
-				BytesProcessed += ManifestResult.FileSize;
-				if (!ManifestResult.bSuccess)
-				{
-					bSuccess = false;
-					UE_LOG(LogDataCompactifier, Error, TEXT("Failed to extract references from %s."), *ManifestFilename);
-				}
-				else
-				{
-					const int32 NumReferences = ManifestResult.References.Num();
-					ReferencedFileSet.Append(MoveTemp(ManifestResult.References));
-					UE_LOG(LogDataCompactifier, Display, TEXT("Extracted %d references from %s. Unioning with existing files, new count of %d."), NumReferences, *ManifestFilename, ReferencedFileSet.Num());
-				}
-				++ManifestFileIdx;
-			}
+			float StatsLoggerTimeSeconds = 60.0f;
+			GConfig->GetFloat(TEXT("BuildPatchServices"), TEXT("StatsLoggerTimeSeconds"), StatsLoggerTimeSeconds, GEngineIni);
 
+			// Handle all manifests first.
+			FProcessingStats ProcessedFilesStat(*StatsCollector, TEXT("Processed Files"));
+			TSet<FReferencedFileData> ReferencedFileSet;
+			bSuccess = FindAllReferencedFiles(StatsLoggerTimeSeconds, FirstNonManifest, AllFiles, ProcessedFilesStat, ReferencedFileSet);
 			if (bSuccess)
 			{
+				FProcessingStats UnreferencedFilesStat(*StatsCollector, TEXT("Unreferenced Files"));
+				FProcessingStats SkippedFilesStat(*StatsCollector, TEXT("Skipped Files"));
+				FProcessingStats NonPatchFilesStat(*StatsCollector, TEXT("Unknown Files"));
+				FProcessingStats DeletedFilesStat(*StatsCollector, TEXT("Deleted Files"));
+
+				FStatsTimeGuard HandleAllCloudFilesTimeStat(*StatsCollector, TEXT("Data Compactifier: Handle Cloud Files Time"));
+
 				// Now handle all CloudDir files.
-				TArray<TFuture<FDataFileResult>> DataFileFutures;
+				FTaskScheduler<FDataFileResult, MaxFileHandlingThreads> DataFileTasks;
+				UE_LOG(LogDataCompactifier, Log, TEXT("Using %d queue depth for file handling operations."), MaxFileHandlingThreads);
 				for (int32 FileIdx = FirstNonManifest; FileIdx < AllFiles.Num(); ++FileIdx)
 				{
 					const FString& Filename = AllFiles[FileIdx];
-					DataFileFutures.Add(Async(EAsyncExecution::ThreadPool, [this, &Filename, &ReferencedFileSet, &Cutoff]()
+					if (ReferencedFileSet.ContainsByHash(GetTypeHash(Filename), Filename))
+					{
+						continue;
+					}
+					auto TaskFunc = [&]()
 					{
-						int64 FileSize = 0;
-						bool bIsReferenced = false;
-						bool bIsOldEnough = false;
-						bool bIsRecognisedFileType = false;
-						bool bShouldDelete = false;
-						const bool bGotFileSize = FileSystem->GetFileSize(*Filename, FileSize);
-						const bool bFileStillExists = FileSystem->FileExists(*Filename);
+						FFileStatData FileStatData = FileSystem->GetStatData(*Filename);
+						const bool bGotFileSize = FileStatData.FileSize >= 0;
+						const bool bFileStillExists = FileStatData.bIsValid;
 						if (bGotFileSize)
 						{
-							bIsReferenced = ReferencedFileSet.Contains(Filename);
-							if (!bIsReferenced)
+							bool bIsOldEnough = FileStatData.ModificationTime < Cutoff;
+							bool bIsRecognisedFileType = IsPatchData(Filename);
+							bool bShouldDelete = bIsOldEnough && bIsRecognisedFileType;
+							if (bShouldDelete)
 							{
-								FDateTime FileTimeStamp;
-								bIsOldEnough = FileSystem->GetTimeStamp(*Filename, FileTimeStamp) && FileTimeStamp < Cutoff;
-								bIsRecognisedFileType = IsPatchData(Filename);
-								bShouldDelete = bIsOldEnough && bIsRecognisedFileType;
-								if (bShouldDelete)
-								{
-									DeleteFile(Filename);
-								}
+								DeleteFile(Filename);
 							}
+							return FDataFileResult(bIsOldEnough, bIsRecognisedFileType, bShouldDelete, FileStatData.FileSize);
 						}
 						else if (bFileStillExists)
 						{
 							UE_LOG(LogDataCompactifier, Warning, TEXT("Could not determine size of %s. Assuming 0 bytes."), *Filename);
-							FileSize = 0;
 						}
 						else
 						{
 							UE_LOG(LogDataCompactifier, Log, TEXT("File removed since enumeration %s. Using 0 bytes."), *Filename);
-							FileSize = 0;
 						}
-						return FDataFileResult(bIsReferenced, bIsOldEnough, bIsRecognisedFileType, bShouldDelete, FileSize);
-					}));
-				}
-				int32 DataFileIdx = FirstNonManifest;
-				for (TFuture<FDataFileResult>& DataFileFuture : DataFileFutures)
-				{
-					// Grab result.
-					FDataFileResult DataFileResult = DataFileFuture.Get();
-					const FString& DataFilename = AllFiles[DataFileIdx];
-					// Process result.
-					++FilesProcessed;
-					BytesProcessed += DataFileResult.FileSize;
-					if (!DataFileResult.bIsReferenced)
+						return FDataFileResult();
+					};
+					auto CompleteFunc = [&](FDataFileResult Result)
 					{
-						if (!DataFileResult.bIsOldEnough)
+						ProcessedFilesStat.Accumulate(Result.FileSize);
+						UnreferencedFilesStat.Accumulate(Result.FileSize);
+						if (!Result.bIsOldEnough)
 						{
-							++FilesSkipped;
-							BytesSkipped += DataFileResult.FileSize;
+							SkippedFilesStat.Accumulate(Result.FileSize);
 						}
-						else if (!DataFileResult.bIsRecognisedFileType)
+						else if (!Result.bIsRecognisedFileType)
 						{
-							++NonPatchFilesProcessed;
-							NonPatchBytesProcessed += DataFileResult.FileSize;
+							NonPatchFilesStat.Accumulate(Result.FileSize);
 						}
-						else if (DataFileResult.bShouldDelete)
+						else if (Result.bShouldDelete)
 						{
-							++FilesDeleted;
-							BytesDeleted += DataFileResult.FileSize;
+							DeletedFilesStat.Accumulate(Result.FileSize);
 							if (DeletedChunkOutput.IsValid())
 							{
-								FString OutputLine = DataFilename + TEXT("\r\n");
+								FString OutputLine = Filename + TEXT("\r\n");
 								FTCHARToUTF8 UTF8String(*OutputLine);
 								DeletedChunkOutput->Serialize((UTF8CHAR*)UTF8String.Get(), UTF8String.Length() * sizeof(UTF8CHAR));
 							}
-							UE_LOG(LogDataCompactifier, Log, TEXT("Deprecated data %s%s."), *DataFilename, Configuration.bRunPreview ? TEXT("") : TEXT(" deleted"));
+							UE_LOG(LogDataCompactifier, Log, TEXT("Deprecated data %s%s."), *Filename, Configuration.bRunPreview ? TEXT("") : TEXT(" deleted"));
 						}
-					}
-					++DataFileIdx;
-				}
+					};
 
-				UE_LOG(LogDataCompactifier, Display, TEXT("Found %u files totalling %s."), FilesProcessed, *BuildSizeString(BytesProcessed));
-				UE_LOG(LogDataCompactifier, Display, TEXT("Found %u unknown files totalling %s."), NonPatchFilesProcessed, *BuildSizeString(NonPatchBytesProcessed));
-				UE_LOG(LogDataCompactifier, Display, TEXT("Deleted %u files totalling %s."), FilesDeleted, *BuildSizeString(BytesDeleted));
-				UE_LOG(LogDataCompactifier, Display, TEXT("Skipped %u young files totalling %s."), FilesSkipped, *BuildSizeString(BytesSkipped));
+					DataFileTasks.ScheduleTask(MoveTemp(TaskFunc), MoveTemp(CompleteFunc));
+					HandleAllCloudFilesTimeStat.Refresh();
+					FStatsCollector::SetAsPercentage(AllFilesProgress, (double)FileIdx / (double)AllFiles.Num());
+					LogStats(StatsLoggerTimeSeconds);
+				}
+				DataFileTasks.WaitAll();
+				FStatsCollector::SetAsPercentage(AllFilesProgress, 1);
+				UE_LOG(LogDataCompactifier, Display, TEXT("Cloud files handling completed in %s."), *HandleAllCloudFilesTimeStat.Refresh().GetAsString());
+				LogStats();
 			}
 		}
 
@@ -307,17 +433,12 @@ namespace BuildPatchServices
 
 	bool FPatchDataCompactifier::IsPatchData(const FString& FilePath) const
 	{
-		const TCHAR* ChunkFile = TEXT(".chunk");
-		const TCHAR* DeltaFile = TEXT(".delta");
-		const TCHAR* LegacyFile = TEXT(".file");
+		static const TCHAR* ChunkFile = TEXT(".chunk");
+		static const TCHAR* DeltaFile = TEXT(".delta");
+		static const TCHAR* LegacyFile = TEXT(".file");
 		return FilePath.EndsWith(ChunkFile) || FilePath.EndsWith(DeltaFile) || FilePath.EndsWith(LegacyFile);
 	}
 
-	FString FPatchDataCompactifier::BuildSizeString(uint64 Size) const
-	{
-		return FString::Printf(TEXT("%s bytes (%s, %s)"), *FText::AsNumber(Size).ToString(), *FText::AsMemory(Size, &SizeFormattingOptions, nullptr, EMemoryUnitStandard::SI).ToString(), *FText::AsMemory(Size, &SizeFormattingOptions, nullptr, EMemoryUnitStandard::IEC).ToString());
-	}
-
 	IPatchDataCompactifier* FPatchDataCompactifierFactory::Create(const FCompactifyConfiguration& Configuration)
 	{
 		return new FPatchDataCompactifier(Configuration);
diff --git a/Engine/Source/Runtime/Online/BuildPatchServices/Private/Core/AsyncHelpers.h b/Engine/Source/Runtime/Online/BuildPatchServices/Private/Core/AsyncHelpers.h
--- a/Engine/Source/Runtime/Online/BuildPatchServices/Private/Core/AsyncHelpers.h
+++ b/Engine/Source/Runtime/Online/BuildPatchServices/Private/Core/AsyncHelpers.h
@@ -108,6 +108,142 @@ namespace BuildPatchServices
 		}
 	}
 
+	template<typename TResultType>
+	struct TTaskSchedulerTraits
+	{
+		using FOnCompleteFunction = TFunction<void(TResultType)>;
+	};
+
+	template<>
+	struct TTaskSchedulerTraits<void>
+	{
+		using FOnCompleteFunction = TFunction<void()>;
+	};
+
+	/**
+	 * A helper class for scheduling a large amount of tasks on the thread pool while conserving event handles.
+	 * The class itself should be called into entirely from one single thread only. For multiple threads with tasks, use
+	 * separate FTaskScheduler instances.
+	 */
+	template<typename TResultType, int32 CompileTimeMaxTasks = 64>
+	class FTaskScheduler
+	{
+	public:
+		using FOnCompleteCb = typename TTaskSchedulerTraits<TResultType>::FOnCompleteFunction;
+
+		FTaskScheduler() = default;
+
+		/**
+		* @param InRuntimeLimitation Runtime limitaion on a number of concurrent tasks
+		*/
+		explicit FTaskScheduler(int32 InRuntimeLimitation):
+			FTaskScheduler(EAsyncExecution::ThreadPool, InRuntimeLimitation)
+		{
+		}
+
+		/**
+		* @param InRuntimeLimitation Runtime limitaion on a number of concurrent tasks 
+		* @param InExecutionType	 The execution method to use, i.e. on Task Graph or in a separate thread.
+		*/
+		FTaskScheduler(EAsyncExecution InExecutionType, int32 InRuntimeLimitation = CompileTimeMaxTasks)
+			: ExecutionType(InExecutionType)
+			, MaxScheduledTasks(FMath::Clamp(InRuntimeLimitation, 1, CompileTimeMaxTasks))
+		{
+		}
+
+		/**
+		 * Adds a scheduled task to the queue. If the queue is full, will wait for a slot to become available, calling the OnComplete for
+		 * the task being replaced.
+		 * @param Task        The function to be ran asynchronously on the thread pool.
+		 * @param OnComplete  The function to be called with completed task result.
+		 *                    OnComplete will be executed from within ScheduleTask(..) or WaitAll().
+		 */
+		void ScheduleTask(TFunction<TResultType()>&& Task, FOnCompleteCb&& OnComplete = nullptr)
+		{
+			if (TaskFutures.Num() < MaxScheduledTasks)
+			{
+				// Space for new slot.
+				TaskFutures.Add(Async(ExecutionType, MoveTemp(Task)));
+				TaskCompleteFuncs.Add(MoveTemp(OnComplete));
+			}
+			else
+			{
+				// Find a finished slot, or complete the next from last search.
+				int32 NewTaskIndex = SlotSearchCount;
+				for (int32 i = 0; i < TaskFutures.Num(); i++)
+				{
+					const int32 TaskIdx = (SlotSearchCount + i) % TaskFutures.Num();
+					if (TaskFutures[TaskIdx].IsReady())
+					{
+						NewTaskIndex = TaskIdx;
+						break;
+					}
+				}
+				SlotSearchCount = (SlotSearchCount + 1) % TaskFutures.Num();
+				// Complete the slot we will be assigning to.
+				CompleteTask(*this, NewTaskIndex);
+				// Overwrite slot.
+				TaskFutures[NewTaskIndex] = Async(ExecutionType, MoveTemp(Task));
+				TaskCompleteFuncs[NewTaskIndex] = MoveTemp(OnComplete);
+			}
+		}
+
+		/**
+		 * Waits and completes all remaining tasks in the queue. Must always be called to ensure that no OnComplete calls are missed.
+		 */
+		void WaitAll()
+		{
+			for (int32 i = 0; i < TaskFutures.Num(); i++)
+			{
+				CompleteTask(*this, i);
+			}
+			TaskFutures.Empty();
+			TaskCompleteFuncs.Empty();
+		}
+
+	private:
+
+		template<typename TTaskResultType, int32 Tasks>
+		void CompleteTask(FTaskScheduler<TTaskResultType, Tasks>& Scheduler, int32 Idx)
+		{
+			TTaskResultType Result = Scheduler.TaskFutures[Idx].Get();
+			if (Scheduler.TaskCompleteFuncs[Idx])
+			{
+				Scheduler.TaskCompleteFuncs[Idx](MoveTemp(Result));
+			}
+		}
+
+		template<int32 Tasks>
+		void CompleteTask(FTaskScheduler<void, Tasks>& Scheduler, int32 Idx)
+		{
+			Scheduler.TaskFutures[Idx].Wait();
+			if (Scheduler.TaskCompleteFuncs[Idx])
+			{
+				Scheduler.TaskCompleteFuncs[Idx]();
+			}
+		}
+
+	private:
+
+		EAsyncExecution ExecutionType = EAsyncExecution::ThreadPool;
+		int32 MaxScheduledTasks = CompileTimeMaxTasks;
+		TArray<TFuture<TResultType>, TFixedAllocator<CompileTimeMaxTasks>> TaskFutures;
+		TArray<FOnCompleteCb, TFixedAllocator<CompileTimeMaxTasks>> TaskCompleteFuncs;
+		int32 SlotSearchCount = 0;
+	};
+
+	constexpr int32 TaskCoresSchedulerThreadsLimitation = 128;
+
+	template<typename ResultType>
+	using FTaskCoresScheduler = FTaskScheduler<ResultType, TaskCoresSchedulerThreadsLimitation>;
+
+	template<typename ResultType = void, int32 MaxThreads = TaskCoresSchedulerThreadsLimitation>
+	static FTaskCoresScheduler<ResultType> CreateTaskSchedulerLimitedByCoresNumber(EAsyncExecution ExecutionType = EAsyncExecution::ThreadPool)
+	{
+		static const int32 NumberOfCores = FPlatformMisc::NumberOfCoresIncludingHyperthreads();
+		return FTaskScheduler<ResultType, MaxThreads>(ExecutionType, NumberOfCores - 1);
+	}
+
 	typedef FThreadSafeCounter64 FThreadSafeInt64;
 	typedef FThreadSafeCounter FThreadSafeInt32;
 }
diff --git a/Engine/Source/Runtime/Online/BuildPatchServices/Private/Core/FileSpan.h b/Engine/Source/Runtime/Online/BuildPatchServices/Private/Core/FileSpan.h
--- a/Engine/Source/Runtime/Online/BuildPatchServices/Private/Core/FileSpan.h
+++ b/Engine/Source/Runtime/Online/BuildPatchServices/Private/Core/FileSpan.h
@@ -8,6 +8,7 @@ namespace BuildPatchServices
 {
 	struct FFileSpan
 	{
+	public:
 		FFileSpan(const FString& InFilename, uint64 InSize, uint64 InStartIdx, bool InIsUnixExecutable, const FString& InSymlinkTarget)
 			: Filename(InFilename)
 			, Size(InSize)
@@ -21,7 +22,10 @@ namespace BuildPatchServices
 			: Filename(CopyFrom.Filename)
 			, Size(CopyFrom.Size)
 			, StartIdx(CopyFrom.StartIdx)
-			, SHAHash(CopyFrom.SHAHash)
+			, MD5Hash(CopyFrom.MD5Hash)
+			, SHA1Hash(CopyFrom.SHA1Hash)
+			, SHA256Hash(CopyFrom.SHA256Hash)
+			, MIMEType(CopyFrom.MIMEType)
 			, IsUnixExecutable(CopyFrom.IsUnixExecutable)
 			, SymlinkTarget(CopyFrom.SymlinkTarget)
 		{
@@ -31,7 +35,10 @@ namespace BuildPatchServices
 			: Filename(MoveTemp(MoveFrom.Filename))
 			, Size(MoveTemp(MoveFrom.Size))
 			, StartIdx(MoveTemp(MoveFrom.StartIdx))
-			, SHAHash(MoveTemp(MoveFrom.SHAHash))
+			, MD5Hash(MoveTemp(MoveFrom.MD5Hash))
+			, SHA1Hash(MoveTemp(MoveFrom.SHA1Hash))
+			, SHA256Hash(MoveTemp(MoveFrom.SHA256Hash))
+			, MIMEType(MoveTemp(MoveFrom.MIMEType))
 			, IsUnixExecutable(MoveTemp(MoveFrom.IsUnixExecutable))
 			, SymlinkTarget(MoveTemp(MoveFrom.SymlinkTarget))
 		{
@@ -49,7 +56,10 @@ namespace BuildPatchServices
 			Filename = CopyFrom.Filename;
 			Size = CopyFrom.Size;
 			StartIdx = CopyFrom.StartIdx;
-			SHAHash = CopyFrom.SHAHash;
+			MD5Hash = CopyFrom.MD5Hash;
+			SHA1Hash = CopyFrom.SHA1Hash;
+			SHA256Hash = CopyFrom.SHA256Hash;
+			MIMEType = CopyFrom.MIMEType;
 			IsUnixExecutable = CopyFrom.IsUnixExecutable;
 			SymlinkTarget = CopyFrom.SymlinkTarget;
 			return *this;
@@ -60,16 +70,23 @@ namespace BuildPatchServices
 			Filename = MoveTemp(MoveFrom.Filename);
 			Size = MoveTemp(MoveFrom.Size);
 			StartIdx = MoveTemp(MoveFrom.StartIdx);
-			SHAHash = MoveTemp(MoveFrom.SHAHash);
+			MD5Hash = MoveTemp(MoveFrom.MD5Hash);
+			SHA1Hash = MoveTemp(MoveFrom.SHA1Hash);
+			SHA256Hash = MoveTemp(MoveFrom.SHA256Hash);
+			MIMEType = MoveTemp(MoveFrom.MIMEType);
 			IsUnixExecutable = MoveTemp(MoveFrom.IsUnixExecutable);
 			SymlinkTarget = MoveTemp(MoveFrom.SymlinkTarget);
 			return *this;
 		}
 
+	public:
 		FString Filename;
 		uint64 Size;
 		uint64 StartIdx;
-		FSHAHash SHAHash;
+		FMD5Hash MD5Hash;
+		FSHAHash SHA1Hash;
+		FSHA256Signature SHA256Hash;
+		FString MIMEType;
 		bool IsUnixExecutable;
 		FString SymlinkTarget;
 	};
diff --git a/Engine/Source/Runtime/Online/BuildPatchServices/Private/Core/RingBuffer.h b/Engine/Source/Runtime/Online/BuildPatchServices/Private/Core/RingBuffer.h
--- a/Engine/Source/Runtime/Online/BuildPatchServices/Private/Core/RingBuffer.h
+++ b/Engine/Source/Runtime/Online/BuildPatchServices/Private/Core/RingBuffer.h
@@ -193,7 +193,12 @@ template< typename DataType >
 void TRingBuffer< DataType >::Enqueue(const DataType& Val)
 {
 	Data[DataIndex++] = Val;
-	DataIndex %= BufferDataSize;
+
+	if (DataIndex == BufferDataSize)
+	{
+		DataIndex = 0;
+	}
+
 	++TotalNumDataPushed;
 	NumDataAvailable = FMath::Clamp< uint32 >(NumDataAvailable + 1, 0, BufferDataSize);
 }
@@ -209,7 +214,10 @@ void TRingBuffer< DataType >::Enqueue(const DataType* ValBuf, const uint32& BufL
 		FMemory::Memcpy(Data, ValBuf + FirstPartLen, sizeof(DataType)* (BufLen - FirstPartLen));
 	}
 	DataIndex += BufLen;
-	DataIndex %= BufferDataSize;
+	if (DataIndex >= BufferDataSize)
+	{
+		DataIndex -= BufferDataSize;
+	}
 	TotalNumDataPushed += BufLen;
 	NumDataAvailable = FMath::Clamp<uint32>(NumDataAvailable + BufLen, 0, BufferDataSize);
 }
