p4 integrate //UE5/Dev-Distro-5.5/Engine/Source/Programs/BuildPatchTool/... //UE5/Main/Engine/Source/Programs/BuildPatchTool/...

diff --git a/Engine/Source/Runtime/Online/BuildPatchServices/Private/Data/ChunkData.cpp b/Engine/Source/Runtime/Online/BuildPatchServices/Private/Data/ChunkData.cpp
--- a/Engine/Source/Runtime/Online/BuildPatchServices/Private/Data/ChunkData.cpp
+++ b/Engine/Source/Runtime/Online/BuildPatchServices/Private/Data/ChunkData.cpp
@@ -6,9 +6,13 @@
 #include "Memory/SharedBuffer.h"
 #include "Misc/CommandLine.h"
 #include "Misc/Compression.h"
+#include "Misc/ConfigCacheIni.h"
+#include "Common/Crypto.h"
 #include "Common/FileSystem.h"
+#include "Common/StatsCollector.h"
 #include "BuildPatchHash.h"
 #include "BuildPatchManifest.h"
+#include "Installer/CloudChunkSource.h"
 #include "BuildPatchServicesPrivate.h"
 
 // The chunk header magic codeword, for quick checking that the opened file is a chunk file.
@@ -17,13 +21,19 @@
 // The chunkdb header magic codeword, for quick checking that the opened file is a chunkdb file.
 #define CHUNKDB_HEADER_MAGIC    0xB1FE3AA3
 
+// Some logs are spammy, but used in automated tests, so default to suppress them in shipping.
+#if UE_BUILD_SHIPPING
 DEFINE_LOG_CATEGORY_STATIC(LogChunkData, Log, All);
+#else
+DEFINE_LOG_CATEGORY_STATIC(LogChunkData, Verbose, All);
+#endif // UE_BUILD_SHIPPING
+
 
 namespace BuildPatchServices
 {
 	const TCHAR* ToString(const EChunkLoadResult& ChunkLoadResult)
 	{
-		switch(ChunkLoadResult)
+		switch (ChunkLoadResult)
 		{
 			case EChunkLoadResult::Success:
 				return TEXT("Success");
@@ -54,7 +64,8 @@ namespace BuildPatchServices
 
 	const TCHAR* ToString(const EChunkSaveResult& ChunkSaveResult)
 	{
-		switch(ChunkSaveResult)
+		static_assert((uint32)EChunkSaveResult::InvalidOrMax == 6, "Please adjust the switch below for new values.");
+		switch (ChunkSaveResult)
 		{
 			case EChunkSaveResult::Success:
 				return TEXT("Success");
@@ -64,11 +75,25 @@ namespace BuildPatchServices
 				return TEXT("BadArchive");
 			case EChunkSaveResult::SerializationError:
 				return TEXT("SerializationError");
+			case EChunkSaveResult::MissingEncryptionKey:
+				return TEXT("MissingEncryptionKey");
+			case EChunkSaveResult::EncryptionFail:
+				return TEXT("EncryptionFail");
 			default:
 				return TEXT("Unknown");
 		}
 	}
 
+	bool ConsideredFailure(const EChunkLoadResult& ChunkLoadResult)
+	{
+		return ChunkLoadResult != EChunkLoadResult::Success && ChunkLoadResult != EChunkLoadResult::MissingSecretKey;
+	}
+
+	bool ConsideredFailure(const EChunkSaveResult& ChunkSaveResult)
+	{
+		return ChunkSaveResult != EChunkSaveResult::Success;
+	}
+
 	/**
 	 * Enum which describes the chunk header version.
 	 */
@@ -78,6 +103,7 @@ namespace BuildPatchServices
 		Original,
 		StoresShaAndHashType,
 		StoresDataSizeUncompressed,
+		StoresEncryptionData,
 
 		// Always after the latest version, signifies the latest version plus 1 to allow initialization simplicity.
 		LatestPlusOne,
@@ -96,9 +122,11 @@ namespace BuildPatchServices
 		// StoresShaAndHashType is 62 bytes (328b Original, 160b SHA1, 8b HashType).
 		62,
 		// StoresDataSizeUncompressed is 66 bytes (496b StoresShaAndHashType, 32b DataSizeUncompressed).
-		66
+		66,
+		// StoresEncryptionData is 98 bytes (496b StoresDataSizeUncompressed, 4x32b GUID, 128b AuthTag).
+		98
 	};
-	static_assert((uint32)EChunkVersion::LatestPlusOne == 4, "Please adjust ChunkHeaderVersionSizes values accordingly.");
+	static_assert((uint32)EChunkVersion::LatestPlusOne == 5, "Please adjust ChunkHeaderVersionSizes values accordingly.");
 
 	namespace HeaderHelpers
 	{
@@ -117,6 +145,7 @@ namespace BuildPatchServices
 
 		EChunkVersion FeatureLevelToChunkVersion(EFeatureLevel FeatureLevel)
 		{
+			static_assert((uint32)EFeatureLevel::Latest == 24, "Please adjust the switch below for new feature levels.");
 			switch (FeatureLevel)
 			{
 				case BuildPatchServices::EFeatureLevel::Original:
@@ -140,20 +169,26 @@ namespace BuildPatchServices
 				case BuildPatchServices::EFeatureLevel::VariableSizeChunks:
 				case BuildPatchServices::EFeatureLevel::UsesRuntimeGeneratedBuildId:
 				case BuildPatchServices::EFeatureLevel::UsesBuildTimeGeneratedBuildId:
+				case BuildPatchServices::EFeatureLevel::StoresFileMD5HashesAndMIMEType:
+				case BuildPatchServices::EFeatureLevel::StoresFileSHA256Hashes:
+				case BuildPatchServices::EFeatureLevel::StoresUninstallActions:
 					return EChunkVersion::StoresDataSizeUncompressed;
+				case BuildPatchServices::EFeatureLevel::ChunkEncryptionSupport:
+				case BuildPatchServices::EFeatureLevel::ChunksStoredBySecret:
+				case BuildPatchServices::EFeatureLevel::ManifestEncryptionSupport:
+					return EChunkVersion::StoresEncryptionData;
 			}
-			checkf(false, TEXT("Unhandled FeatureLevel %s"), FeatureLevelToString(FeatureLevel));
+			checkf(false, TEXT("Unhandled FeatureLevel %s"), LexToString(FeatureLevel));
 			return EChunkVersion::Invalid;
 		}
 	}
-	static_assert((uint32)EFeatureLevel::Latest == 18, "Please adjust HeaderHelpers::FeatureLevelToChunkVersion for new feature levels.");
 
 	FChunkHeader::FChunkHeader()
 		: Version((uint32)EChunkVersion::Latest)
 		, HeaderSize(ChunkHeaderVersionSizes[(uint32)EChunkVersion::Latest])
 		, DataSizeCompressed(0)
 		, DataSizeUncompressed(1024 * 1024)
-		, StoredAs(EChunkStorageFlags::None)
+		, StoredAs(EChunkStorageFlags::RawData)
 		, HashType(EChunkHashFlags::RollingPoly64)
 		, RollingHash(0)
 	{
@@ -213,6 +248,19 @@ namespace BuildPatchServices
 				}
 				ExpectedSerializedBytes = ChunkHeaderVersionSizes[(uint32)EChunkVersion::StoresDataSizeUncompressed];
 			}
+
+			// Serialse encryption ID and auth tag required by this chunk
+			if (bSuccess && Header.Version >= (uint32)EChunkVersion::StoresEncryptionData)
+			{
+				bSuccess = Ar.IsSaving() || (ArchiveSizeLeft >= ChunkHeaderVersionSizes[(uint32)EChunkVersion::StoresEncryptionData]);
+				if (bSuccess)
+				{
+					Ar << Header.EncryptionSecretId;
+					Ar.Serialize(Header.AESAuthTag.AuthTag, AES256_GCM_AuthTagSizeInBytes);
+					bSuccess = !Ar.IsError();
+				}
+				ExpectedSerializedBytes = ChunkHeaderVersionSizes[(uint32)EChunkVersion::StoresEncryptionData];
+			}
 		}
 
 		// Make sure the expected number of bytes were serialized. In practice this will catch errors where type
@@ -241,18 +289,6 @@ namespace BuildPatchServices
 		return Ar;
 	}
 
-	/* FChunkInfo implementation
-	*****************************************************************************/
-	FChunkInfo::FChunkInfo()
-		: Guid()
-		, Hash(0)
-		, ShaHash()
-		, GroupNumber(0)
-		, WindowSize(1048576)
-		, FileSize(0)
-	{
-	}
-
 	/* FChunkPart implementation
 	*****************************************************************************/
 	FChunkPart::FChunkPart()
@@ -491,11 +527,26 @@ namespace BuildPatchServices
 
 	class FChunkDataSerialization : public IChunkDataSerialization
 	{
+		int64 MinCompressionBytes = 64 * 1024;
+		int32 MinCompressionPct = 10;
+
 	public:
-		FChunkDataSerialization(IFileSystem* InFileSystem, EFeatureLevel InFeatureLevel)
+		FChunkDataSerialization(IFileSystem* InFileSystem, ICrypto* InCrypto, FChunkDataSerializationConfig InConfig)
 			: FileSystem(InFileSystem)
-			, FeatureLevel(InFeatureLevel)
+			, Crypto(InCrypto)
+			, Config(InConfig)
+			, WriteEncryptionSecretKey(Config.EncryptionSecrets.Find(Config.WriteEncryptionSecretId))
+			// UE5 MERGE TODO: , CloudChunkSourceStat(Config.CloudChunkSourceStat)
 		{
+			// UE5 MERGE TODO: check(CloudChunkSourceStat);
+			GConfig->GetInt64(TEXT("BuildPatchTool"), TEXT("MinCompressionBytes"), MinCompressionBytes, GEngineIni);
+			GConfig->GetInt(TEXT("BuildPatchTool"), TEXT("MinCompressionPct"), MinCompressionPct, GEngineIni);
+
+			FParse::Value(FCommandLine::Get(), TEXT("MinCompressionPct="), MinCompressionPct);
+			FParse::Value(FCommandLine::Get(), TEXT("MinCompressionBytes="), MinCompressionBytes);
+
+			UE_LOG(LogChunkData, Display, TEXT("Using MinCompressionPct=%d"), MinCompressionPct);
+			UE_LOG(LogChunkData, Display, TEXT("Using MinCompressionBytes=%d"), MinCompressionBytes);
 		}
 
 		~FChunkDataSerialization()
@@ -529,7 +580,6 @@ namespace BuildPatchServices
 
 			if (!OutHeader.Guid.IsValid() || // Must have a valid guid
 				(!EnumHasAnyFlags(OutHeader.HashType, EChunkHashFlags::RollingPoly64) && !EnumHasAnyFlags(OutHeader.HashType, EChunkHashFlags::Sha1)) || // Must have a valid hash
-				((OutHeader.StoredAs & EChunkStorageFlags::Encrypted) != EChunkStorageFlags::None) || // Must not be encrypted
 				(InDestinationBuffer.GetSize() > OutHeader.DataSizeUncompressed)) // Unexpected buffer size.
 			{
 				UE_LOG(LogBuildPatchServices, Error, TEXT("Failed to validate chunk. Header reported Guid=%s HashType=%d StoredAs=%d UncompressedSize=%d DestinationSize=%llu"),
@@ -545,7 +595,7 @@ namespace BuildPatchServices
 			// We seem to be getting chunkinfos that report window sizes that are not the full chunk size, so we need to be
 			// able to catch these subchunks and copy them out as needed.
 			uint64 ReadSize;
-			if ((OutHeader.StoredAs & EChunkStorageFlags::Compressed) == EChunkStorageFlags::None &&
+			if (OutHeader.StoredAs == EChunkStorageFlags::RawData &&
 				InDestinationBuffer.GetSize() == OutHeader.DataSizeUncompressed)
 			{
 				ReadSize = InDestinationBuffer.GetSize();
@@ -553,7 +603,7 @@ namespace BuildPatchServices
 			}
 			else
 			{
-				// We need to allocate and read in to another buffer for decompression.
+				// We need to allocate and read in to another buffer for decompression/decryption.
 				ReadSize = OutHeader.DataSizeCompressed;
 				OutCompressedBuffer = FUniqueBuffer::Alloc(OutHeader.DataSizeCompressed);
 				InArchive.Serialize(OutCompressedBuffer.GetData(), OutCompressedBuffer.GetSize());
@@ -567,72 +617,77 @@ namespace BuildPatchServices
 			return true;
 		}
 
-
-		virtual bool DecompressValidatedRead(const FChunkHeader& InHeader, FMutableMemoryView InDestinationBuffer, const FUniqueBuffer& InCompressedBuffer) const override
+		virtual bool DecompressAndDecryptValidatedRead(const FChunkHeader& InHeader, FMutableMemoryView InDestinationBuffer, const FUniqueBuffer& InCompressedBuffer) const override
 		{
-			// ValidateAndRead sets up the data we get on the dispatch thread so we can decompress and hash here
-			// in another thread. Where the data lives depends on compression and whether the read requested the 
-			// full buffer.
-			
+			FMutableMemoryView UseDestinationBuffer = InDestinationBuffer;
+			FMemoryView UseCompressedBuffer = InCompressedBuffer.GetView();
+			TArray<uint8> DecryptedData;
+
+			FUniqueBuffer TempDestinationBuffer;
+			if (InHeader.DataSizeUncompressed != InDestinationBuffer.GetSize())
+			{
+				TempDestinationBuffer = FUniqueBuffer::Alloc(InHeader.DataSizeUncompressed);
+				UseDestinationBuffer = TempDestinationBuffer.GetView();
+			}
 
-			FMemoryView FullChunkBuffer;
-			FUniqueBuffer DecompressionBuffer;
-			if ((InHeader.StoredAs & EChunkStorageFlags::Compressed) != EChunkStorageFlags::None)
+			// Decrypt if necessary	
+			if (EnumHasAllFlags(InHeader.StoredAs, EChunkStorageFlags::Encrypted))
 			{
-				// Compressed - we usually can decompress directly to the output, but sometimes
-				// the request isn't for the full chunk so we use a temp.
-				FMutableMemoryView DecompressTo = InDestinationBuffer;
-				if (DecompressTo.GetSize() != InHeader.DataSizeUncompressed)
+				const TArray<uint8>* EncryptionSecretKey = Config.EncryptionSecrets.Find(InHeader.EncryptionSecretId);
+
+				if (EncryptionSecretKey == nullptr)
 				{
-					DecompressionBuffer = FUniqueBuffer::Alloc(InHeader.DataSizeUncompressed);
-					DecompressTo = DecompressionBuffer.GetView();
+					UE_LOG(LogChunkData, Error, TEXT("Decryption failed for chunk: %s EncryptionSecretKey is empty."), *InHeader.Guid.ToString());
+					return false;
 				}
 
-				if (!FCompression::UncompressMemory(
-					NAME_Zlib,
-					DecompressTo.GetData(),
-					DecompressTo.GetSize(),
-					InCompressedBuffer.GetData(),
-					InCompressedBuffer.GetSize()))
+				// Attempt to decrypt the data
+				const uint8* CompressedData = (const uint8*)InCompressedBuffer.GetData();
+				TArrayView<const uint8> Ciphertext(CompressedData, InCompressedBuffer.GetSize());
+				TArrayView<const uint8> InitializationVector(InHeader.SHAHash.Hash, AES256_GCM_InitializationVectorSizeInBytes);
+				TArrayView<const uint8> AESAuthTag(InHeader.AESAuthTag.AuthTag, AES256_GCM_AuthTagSizeInBytes);
+				bool bDecryptResult = false;
+				uint64 DecryptStartCycles = FStatsCollector::GetCycles();
+				DecryptedData = Crypto->Decrypt_AES_256_GCM(Ciphertext, *EncryptionSecretKey, InitializationVector, AESAuthTag, bDecryptResult);
+				// UE5 MERGE TODO: CloudChunkSourceStat->OnChunkDecrypted(FStatsCollector::GetCycles() - DecryptStartCycles);
+				if (bDecryptResult)
+				{
+					UseCompressedBuffer = FMemoryView(DecryptedData.GetData(), DecryptedData.Num());
+				}
+				else
 				{
-					// already logged.
+					UE_LOG(LogChunkData, Error, TEXT("Decryption failed for chunk: %s (EncryptionSecretId=\"%s\")."), *InHeader.Guid.ToString(), *InHeader.EncryptionSecretId.ToString());
 					return false;
 				}
-
-				FullChunkBuffer = DecompressTo;
 			}
-			else
+
+			if (EnumHasAnyFlags(InHeader.StoredAs, EChunkStorageFlags::Compressed))
 			{
-				// Uncompressed - where ValidateAndRead put the data depends on whether the destination buffer
-				// was big enough for the full chunk.
-				FullChunkBuffer = InDestinationBuffer;
-				if (InDestinationBuffer.GetSize() != InHeader.DataSizeUncompressed)
+				if (!FCompression::UncompressMemory(
+					NAME_Zlib,
+					UseDestinationBuffer.GetData(),
+					UseDestinationBuffer.GetSize(),
+					UseCompressedBuffer.GetData(),
+					UseCompressedBuffer.GetSize()))
 				{
-					// The data was read into InCompressedBuffer instead.
-					FullChunkBuffer = InCompressedBuffer.GetView();
+					return false;
 				}
 			}
-
-			if (!CheckChunkHash(InHeader, FullChunkBuffer.GetData(), FullChunkBuffer.GetSize()))
+			else if (EnumHasAnyFlags(InHeader.StoredAs, EChunkStorageFlags::Encrypted))
 			{
-				UE_LOG(LogChunkData, Error, TEXT("Failed hash check for chunk %s"), *WriteToString<64>(InHeader.Guid));
-				return false;
+				// Encrypted but not compressed, copy before checking hashes
+				FMemory::Memcpy(UseDestinationBuffer.GetData(), UseCompressedBuffer.GetData(), FMath::Min(UseDestinationBuffer.GetSize(), UseCompressedBuffer.GetSize()));
 			}
 
-			// If we had to use a temp buffer, copy to the dest.
-			if (FullChunkBuffer.GetData() != InDestinationBuffer.GetData())
+			bool bHashMatches = CheckChunkHash(InHeader, UseDestinationBuffer.GetData(), UseDestinationBuffer.GetSize());
+			if (bHashMatches && InHeader.DataSizeUncompressed != InDestinationBuffer.GetSize())
 			{
-				// Destination buffer shuold always be <= the chunk size.
-				if (InDestinationBuffer.GetSize() > FullChunkBuffer.GetSize())
-				{
-					UE_LOG(LogChunkData, Error, TEXT("Invalid destination buffer size for chunk %s"), *WriteToString<64>(InHeader.Guid));
-					return false;
-				}
-				FMemory::Memcpy(InDestinationBuffer.GetData(), FullChunkBuffer.GetData(), InDestinationBuffer.GetSize());
+				// We routed through a temp buffer, copy out as necessary.
+				FMemory::Memcpy(InDestinationBuffer.GetData(), TempDestinationBuffer.GetData(), InDestinationBuffer.GetSize());
 			}
-			return true;
-		}
 
+			return bHashMatches;
+		}
 
 		virtual IChunkDataAccess* LoadFromFile(const FString& Filename, EChunkLoadResult& OutLoadResult) const override
 		{
@@ -673,8 +728,19 @@ namespace BuildPatchServices
 				return nullptr;
 			}
 		}
+		virtual TFunction<TPair<EChunkLoadResult, IChunkDataAccess*>()> LoadFromArchiveAsync(FArchive& Archive) const override
+		{
+			if (Archive.IsLoading())
+			{
+				return LoadAsync(Archive);
+			}
+			else
+			{
+				return []() {return TPair<EChunkLoadResult, IChunkDataAccess*>(EChunkLoadResult::BadArchive, nullptr); };
+			}
+		}
 
-		virtual EChunkSaveResult SaveToFile(const FString& Filename, const IChunkDataAccess* ChunkDataAccess) const override
+		virtual EChunkSaveResult SaveToFile(const FString& Filename, IChunkDataAccess* ChunkDataAccess) const override
 		{
 			EChunkSaveResult SaveResult;
 			TUniquePtr<FArchive> FileOut(FileSystem->CreateFileWriter(*Filename));
@@ -689,13 +755,13 @@ namespace BuildPatchServices
 			return SaveResult;
 		}
 
-		virtual EChunkSaveResult SaveToMemory(TArray<uint8>& Memory, const IChunkDataAccess* ChunkDataAccess) const override
+		virtual EChunkSaveResult SaveToMemory(TArray<uint8>& Memory, IChunkDataAccess* ChunkDataAccess) const override
 		{
 			FMemoryWriter MemoryWriter(Memory);
 			return Save(MemoryWriter, ChunkDataAccess, true);
 		}
 
-		virtual EChunkSaveResult SaveToArchive(FArchive& Archive, const IChunkDataAccess* ChunkDataAccess) const override
+		virtual EChunkSaveResult SaveToArchive(FArchive& Archive, IChunkDataAccess* ChunkDataAccess) const override
 		{
 			if (Archive.IsSaving())
 			{
@@ -707,7 +773,7 @@ namespace BuildPatchServices
 			}
 		}
 
-		virtual EChunkSaveResult SaveToArchiveUncompressed(FArchive& Archive, const IChunkDataAccess* ChunkDataAccess) const override
+		virtual EChunkSaveResult SaveToArchiveUncompressed(FArchive& Archive, IChunkDataAccess* ChunkDataAccess) const override
 		{
 			if (Archive.IsSaving())
 			{
@@ -719,33 +785,43 @@ namespace BuildPatchServices
 			}
 		}
 
-		virtual void InjectShaToChunkData(TArray<uint8>& Memory, const FSHAHash& ShaHashData) const override
+		virtual bool InjectShaToChunkData(TArray<uint8>& Memory, const FSHAHash& ShaHashData) const override
 		{
 			const uint32 StoresShaAndHashTypeUint = (uint32)EChunkVersion::StoresShaAndHashType;
-			const uint32 StoresShaAndHashTypeHeaderSize = ChunkHeaderVersionSizes[StoresShaAndHashTypeUint];
 			FMemoryReader MemoryReader(Memory);
-			FMemoryWriter MemoryWriter(Memory);
 
 			FChunkHeader Header;
 			MemoryReader << Header;
-			Header.HashType |= BuildPatchServices::EChunkHashFlags::Sha1;
+			if (MemoryReader.IsError())
+			{
+				return false;
+			}
+			// Chunk already contains SHA by itself and doesn't need injection
+			if (EnumHasAnyFlags(Header.HashType, EChunkHashFlags::Sha1))
+			{
+				return true;
+			}
+
+			FMemoryWriter MemoryWriter(Memory);
+			EnumAddFlags(Header.HashType, EChunkHashFlags::Sha1);
 			Header.SHAHash = ShaHashData;
 			if (Header.Version < StoresShaAndHashTypeUint)
 			{
+				const uint32 StoresShaAndHashTypeHeaderSize = ChunkHeaderVersionSizes[StoresShaAndHashTypeUint];
 				check(Header.HeaderSize <= StoresShaAndHashTypeHeaderSize);
 				Header.Version = StoresShaAndHashTypeUint;
 				Memory.InsertZeroed(0, StoresShaAndHashTypeHeaderSize - Header.HeaderSize);
 			}
 			MemoryWriter << Header;
+			return !MemoryWriter.IsError();
 		}
 
 		// IChunkDataAccess interface end.
 
 	private:
-		IChunkDataAccess* Load(FArchive& Reader, EChunkLoadResult& OutLoadResult) const
+		TFunction<TPair<EChunkLoadResult, IChunkDataAccess*>()> LoadAsync(FArchive& Reader) const
 		{
-			IChunkDataAccess* ChunkData = nullptr;
-
+			EChunkLoadResult LoadResult = EChunkLoadResult::Success;
 			// Begin of read pos.
 			const int64 StartPos = Reader.Tell();
 
@@ -768,154 +844,217 @@ namespace BuildPatchServices
 					{
 						if ((HeaderCheck.HeaderSize + HeaderCheck.DataSizeCompressed) <= AvailableSize)
 						{
-							if ((HeaderCheck.StoredAs & EChunkStorageFlags::Encrypted) == EChunkStorageFlags::None)
+							constexpr EChunkStorageFlags AllSupportedStorage = EChunkStorageFlags::Compressed | EChunkStorageFlags::Encrypted;
+							if ((HeaderCheck.StoredAs & AllSupportedStorage) == HeaderCheck.StoredAs)
 							{
 								// Create the data.
-								ChunkData = FChunkDataAccessFactory::Create(FileSize);
-
-								// Lock data.
-								FChunkHeader* Header;
-								uint8* Data;
-								ChunkData->GetDataLock(&Data, &Header);
-								*Header = HeaderCheck;
+								IChunkDataAccess* ChunkData = FChunkDataAccessFactory::Create(FileSize);
 
-								// Read the data.
 								{
-									TRACE_CPUPROFILER_EVENT_SCOPE(Chunk_Serialize);
-									Reader.Serialize(Data, Header->DataSizeCompressed);
+									// Lock data.
+									FChunkHeader* Header;
+									uint8* Data;
+									ChunkData->GetDataLock(&Data, &Header);
+									*Header = HeaderCheck; // Update header data in chunk.
+
+									// Read the data.
+									{
+										TRACE_CPUPROFILER_EVENT_SCOPE(Chunk_Serialize);
+										Reader.Serialize(Data, Header->DataSizeCompressed);
+									}
+
+									ChunkData->ReleaseDataLock(); // Should be unlocked by the same thread that locked it. 
 								}
 								if (Reader.IsError() == false)
 								{
-									OutLoadResult = EChunkLoadResult::Success;
-									// Decompress.
-									if ((Header->StoredAs & EChunkStorageFlags::Compressed) != EChunkStorageFlags::None)
+									TFunction<TPair<EChunkLoadResult, IChunkDataAccess*>()> DecryptAndDecodeTask =
+										[ChunkData, this]() mutable -> TPair<EChunkLoadResult, IChunkDataAccess*>
 									{
-										// Create a new data instance.
-										IChunkDataAccess* NewChunkData = FChunkDataAccessFactory::Create(Header->DataSizeUncompressed);
 										// Lock data.
-										FChunkHeader* NewHeader;
-										uint8* NewData;
-										NewChunkData->GetDataLock(&NewData, &NewHeader);
-										// Uncompress the memory.
-										bool bSuccess = FCompression::UncompressMemory(
-											NAME_Zlib,
-											NewData,
-											Header->DataSizeUncompressed,
-											Data,
-											Header->DataSizeCompressed);
-										// If successful, switch over to new data.
-										if (bSuccess)
+										FChunkHeader* Header;
+										uint8* Data;
+										ChunkData->GetDataLock(&Data, &Header);
+
+										EChunkLoadResult LoadResult = EChunkLoadResult::Success;
+										// Decrypt.
+										if (EnumHasAllFlags(Header->StoredAs, EChunkStorageFlags::Encrypted))
 										{
-											FMemory::Memcpy(*NewHeader, *Header);
-											NewHeader->StoredAs = EChunkStorageFlags::None;
-											NewHeader->DataSizeCompressed = Header->DataSizeUncompressed;
-											ChunkData->ReleaseDataLock();
-											delete ChunkData;
-											ChunkData = NewChunkData;
-											Header = NewHeader;
-											Data = NewData;
+											const TArray<uint8>* EncryptionSecretKey = Config.EncryptionSecrets.Find(Header->EncryptionSecretId);
+											if (EncryptionSecretKey != nullptr)
+											{
+												// Attempt to decrypt the data
+												TArrayView<const uint8> Ciphertext(Data, Header->DataSizeCompressed);
+												TArrayView<const uint8> InitializationVector(Header->SHAHash.Hash, AES256_GCM_InitializationVectorSizeInBytes);
+												TArrayView<const uint8> AESAuthTag(Header->AESAuthTag.AuthTag, AES256_GCM_AuthTagSizeInBytes);
+												bool bDecryptResult = false;
+												uint64 DecryptStartCycles = FStatsCollector::GetCycles();
+												TArray<uint8> DecryptedData = Crypto->Decrypt_AES_256_GCM(Ciphertext, *EncryptionSecretKey, InitializationVector, AESAuthTag, bDecryptResult);
+												// UE5 MERGE TODO: CloudChunkSourceStat->OnChunkDecrypted(FStatsCollector::GetCycles() - DecryptStartCycles);
+												if (bDecryptResult)
+												{
+													// If successful, remove the encryption flag
+													EnumRemoveFlags(Header->StoredAs, EChunkStorageFlags::Encrypted);
+													// Copy the decrypted data
+													FMemory::Memcpy(Data, DecryptedData.GetData(), Header->DataSizeCompressed);
+												}
+												else
+												{
+													LoadResult = EChunkLoadResult::DecryptionFailure;
+													UE_LOG(LogChunkData, Error, TEXT("Decryption failed for chunk: %s (EncryptionSecretId=\"%s\")."), *Header->Guid.ToString(), *Header->EncryptionSecretId.ToString());
+												}
+											}
+											else
+											{
+												LoadResult = EChunkLoadResult::MissingSecretKey;
+												UE_LOG(LogChunkData, Verbose, TEXT("The security key is not ready for chunk: %s (EncryptionSecretId=\"%s\")."), *Header->Guid.ToString(), *Header->EncryptionSecretId.ToString());
+											}
 										}
-										// Otherwise delete new data.
-										else
+										// Decompress.
+										if (LoadResult == EChunkLoadResult::Success && EnumHasAllFlags(Header->StoredAs, EChunkStorageFlags::Compressed))
 										{
-											OutLoadResult = EChunkLoadResult::DecompressFailure;
-											NewChunkData->ReleaseDataLock();
-											delete NewChunkData;
-											NewChunkData = nullptr;
+											// Create a new data instance.
+											IChunkDataAccess* NewChunkData = FChunkDataAccessFactory::Create(Header->DataSizeUncompressed);
+											// Lock data.
+											FChunkHeader* NewHeader;
+											uint8* NewData;
+											NewChunkData->GetDataLock(&NewData, &NewHeader);
+											// Uncompress the memory.
+											bool bSuccess = FCompression::UncompressMemory(
+												NAME_Zlib,
+												NewData,
+												Header->DataSizeUncompressed,
+												Data,
+												Header->DataSizeCompressed);
+											// If successful, switch over to new data.
+											if (bSuccess)
+											{
+												FMemory::Memcpy(*NewHeader, *Header);
+												EnumRemoveFlags(NewHeader->StoredAs, EChunkStorageFlags::Compressed);
+												NewHeader->DataSizeCompressed = Header->DataSizeUncompressed;
+												ChunkData->ReleaseDataLock();
+												delete ChunkData;
+												ChunkData = NewChunkData;
+												Header = NewHeader;
+												Data = NewData;
+											}
+											// Otherwise delete new data.
+											else
+											{
+												LoadResult = EChunkLoadResult::DecompressFailure;
+												NewChunkData->ReleaseDataLock();
+												delete NewChunkData;
+												NewChunkData = nullptr;
+											}
 										}
-									}
-									// Verify.
-
-									// If we have both hashes, only check Sha1 instead of both.
-									bool bHasRollingHash = EnumHasAnyFlags(Header->HashType, EChunkHashFlags::RollingPoly64);
-									bool bHasSha1Hash = EnumHasAnyFlags(Header->HashType, EChunkHashFlags::Sha1);
-
-									if (OutLoadResult == EChunkLoadResult::Success && bHasRollingHash && !bHasSha1Hash) // Don't bother verifying the rolling hash if we have a sha1
-									{
-										if (Header->DataSizeCompressed != Header->DataSizeUncompressed || Header->RollingHash != FRollingHash::GetHashForDataSet(Data, Header->DataSizeUncompressed))
+										// Verify.
+										// If we have both hashes, only check Sha1 instead of both.
+										bool bHasRollingHash = EnumHasAnyFlags(Header->HashType, EChunkHashFlags::RollingPoly64);
+										bool bHasSha1Hash = EnumHasAnyFlags(Header->HashType, EChunkHashFlags::Sha1);
+										if (LoadResult == EChunkLoadResult::Success && bHasRollingHash && !bHasSha1Hash) // Don't bother verifying the rolling hash if we have a sha1
 										{
-											OutLoadResult = EChunkLoadResult::HashCheckFailed;
+											if (Header->DataSizeCompressed != Header->DataSizeUncompressed || Header->RollingHash != FRollingHash::GetHashForDataSet(Data, Header->DataSizeUncompressed))
+											{
+												LoadResult = EChunkLoadResult::HashCheckFailed;
+											}
 										}
-									}
-									FSHAHash ShaHashCheck;
-									if (OutLoadResult == EChunkLoadResult::Success && bHasSha1Hash)
-									{
-										TRACE_CPUPROFILER_EVENT_SCOPE(Chunk_ShaHash);
-										FSHA1::HashBuffer(Data, Header->DataSizeUncompressed, ShaHashCheck.Hash);
-										if (!(ShaHashCheck == Header->SHAHash))
+										FSHAHash ShaHashCheck;
+										if (LoadResult == EChunkLoadResult::Success && bHasSha1Hash)
 										{
-											OutLoadResult = EChunkLoadResult::HashCheckFailed;
+											TRACE_CPUPROFILER_EVENT_SCOPE(Chunk_ShaHash);
+											FSHA1::HashBuffer(Data, Header->DataSizeUncompressed, ShaHashCheck.Hash);
+											if (!(ShaHashCheck == Header->SHAHash))
+											{
+												LoadResult = EChunkLoadResult::HashCheckFailed;
+											}
 										}
-									}
+										if (ChunkData != nullptr)
+										{
+											// Release data.
+											ChunkData->ReleaseDataLock();
+	
+											// Delete data if failed
+											if (ConsideredFailure(LoadResult))
+											{
+												delete ChunkData;
+												ChunkData = nullptr;
+											}
+										}
+	
+										return TPair<EChunkLoadResult, IChunkDataAccess*>(LoadResult, ChunkData);
+									};
+									return DecryptAndDecodeTask;
 								}
 								else
 								{
-									OutLoadResult = EChunkLoadResult::SerializationError;
+									LoadResult = EChunkLoadResult::SerializationError;
+									delete ChunkData;
+									ChunkData = nullptr;
 								}
 							}
 							else
 							{
-								OutLoadResult = EChunkLoadResult::UnsupportedStorage;
+								LoadResult = EChunkLoadResult::UnsupportedStorage;
 							}
 						}
 						else
 						{
-							OutLoadResult = EChunkLoadResult::IncorrectFileSize;
+							LoadResult = EChunkLoadResult::IncorrectFileSize;
 						}
 					}
 					else
 					{
-						OutLoadResult = EChunkLoadResult::MissingHashInfo;
+						LoadResult = EChunkLoadResult::MissingHashInfo;
 					}
 				}
 				else
 				{
-					OutLoadResult = EChunkLoadResult::CorruptHeader;
+					LoadResult = EChunkLoadResult::CorruptHeader;
 				}
 			}
 			else
 			{
-				OutLoadResult = EChunkLoadResult::FileSizeTooBig;
+				LoadResult = EChunkLoadResult::FileSizeTooBig;
 			}
 
-			if (ChunkData != nullptr)
-			{
-				// Release data.
-				ChunkData->ReleaseDataLock();
-
-				// Delete data if failed
-				if (OutLoadResult != EChunkLoadResult::Success)
-				{
-					delete ChunkData;
-					ChunkData = nullptr;
-				}
-			}
+			return [LoadResult]() {return TPair<EChunkLoadResult, IChunkDataAccess*>(LoadResult, nullptr); };
+		}
 
-			return ChunkData;
+		IChunkDataAccess* Load(FArchive& Reader, EChunkLoadResult& OutLoadResult) const
+		{
+			const TFunction<TPair<EChunkLoadResult, IChunkDataAccess*>()> DecodeAndDecompressTask = LoadAsync(Reader);
+			const TPair<EChunkLoadResult, IChunkDataAccess*> Result(DecodeAndDecompressTask());
+			OutLoadResult = Result.Key;
+			return Result.Value;
 		}
 
-		EChunkSaveResult Save(FArchive& Writer, const IChunkDataAccess* ChunkDataAccess, bool bTryCompress) const
+		EChunkSaveResult Save(FArchive& Writer, IChunkDataAccess* ChunkDataAccess, bool bTryCompress) const
 		{
-			EChunkSaveResult SaveResult;
-			const uint8* ChunkDataSource;
-			const FChunkHeader* ChunkAccessHeader;
+			EChunkSaveResult SaveResult = EChunkSaveResult::Success;
+			uint8* ChunkDataSource;
+			FChunkHeader* ChunkAccessHeader;
 			ChunkDataAccess->GetDataLock(&ChunkDataSource, &ChunkAccessHeader);
 			// Setup to handle compression.
 			bool bDataIsCompressed = false;
 			TArray<uint8> TempCompressedData;
-			int32 CompressedSize = ChunkAccessHeader->DataSizeUncompressed;
-			if (FeatureLevel >= EFeatureLevel::ChunkCompressionSupport && bTryCompress)
+			int64 CompressedSize = ChunkAccessHeader->DataSizeUncompressed;
+			if (Config.FeatureLevel >= EFeatureLevel::ChunkCompressionSupport && bTryCompress)
 			{
 				TempCompressedData.Empty(ChunkAccessHeader->DataSizeUncompressed);
 				TempCompressedData.AddUninitialized(ChunkAccessHeader->DataSizeUncompressed);
-				// Compression can increase data size, too. This call will return false in that case.
-				bDataIsCompressed = FCompression::CompressMemory(
+				// Compression can increase data size, too. This call will set bDataIsCompressed false in that case.
+				const bool bCompressSuccess = FCompression::CompressMemoryIfWorthDecompressing
+				(
 					NAME_Zlib,
+					bDataIsCompressed,
+					MinCompressionBytes,
+					MinCompressionPct,
 					TempCompressedData.GetData(),
 					CompressedSize,
 					ChunkDataSource,
 					ChunkAccessHeader->DataSizeUncompressed,
-					COMPRESS_BiasMemory);
+					COMPRESS_BiasMemory
+				);
+				bDataIsCompressed = bDataIsCompressed && bCompressSuccess;
 			}
 
 			// If compression succeeded, set data vars.
@@ -928,34 +1067,62 @@ namespace BuildPatchServices
 				CompressedSize = ChunkAccessHeader->DataSizeUncompressed;
 			}
 
-			// Setup Header.
-			FChunkHeader Header;
-			FMemory::Memcpy(Header, *ChunkAccessHeader);
-			Header.Version = (uint32)HeaderHelpers::FeatureLevelToChunkVersion(FeatureLevel);
-			const int64 StartPos = Writer.Tell();
-			Writer << Header;
-			Header.HeaderSize = Writer.Tell() - StartPos;
-			Header.StoredAs = bDataIsCompressed ? EChunkStorageFlags::Compressed : EChunkStorageFlags::None;
-			Header.DataSizeCompressed = CompressedSize;
-			Header.DataSizeUncompressed = ChunkAccessHeader->DataSizeUncompressed;
-			// Make sure we at least have rolling hash.
-			if ((Header.HashType & EChunkHashFlags::RollingPoly64) == EChunkHashFlags::None)
+			TArray<uint8> TempEncryptedData;
+			TArray<uint8> TempEncryptionAuthTag;
+			bool bDataIsEncrypted = Config.FeatureLevel >= EFeatureLevel::ChunkEncryptionSupport && Config.WriteEncryptionSecretId.IsValid();
+			if (bDataIsEncrypted)
 			{
-				Header.HashType = EChunkHashFlags::RollingPoly64;
+				if (WriteEncryptionSecretKey != nullptr)
+				{
+					TArrayView<const uint8> SourceData(ChunkDataSource, CompressedSize);
+					TArrayView<const uint8> InitializationVector(ChunkAccessHeader->SHAHash.Hash, AES256_GCM_InitializationVectorSizeInBytes);
+					bool bEncryptResult = false;
+					TempEncryptedData = Crypto->Encrypt_AES_256_GCM(SourceData, *WriteEncryptionSecretKey, InitializationVector, TempEncryptionAuthTag, bEncryptResult);
+					if (bEncryptResult)
+					{
+						ChunkDataSource = TempEncryptedData.GetData();
+					}
+					else
+					{
+						SaveResult = EChunkSaveResult::EncryptionFail;
+					}
+				}
+				else
+				{
+					SaveResult = EChunkSaveResult::MissingEncryptionKey;
+				}
 			}
 
-			// Write out files.
-			Writer.Seek(StartPos);
-			Writer << Header;
-			// We must use const_cast here due to FArchive::Serialize supporting both read and write. We created a writer.
-			Writer.Serialize(const_cast<uint8*>(ChunkDataSource), CompressedSize);
-			if (!Writer.IsError())
-			{
-				SaveResult = EChunkSaveResult::Success;
-			}
-			else
+			if (SaveResult == EChunkSaveResult::Success)
 			{
-				SaveResult = EChunkSaveResult::SerializationError;
+				// Setup Header.
+				ChunkAccessHeader->Version = (uint32)HeaderHelpers::FeatureLevelToChunkVersion(Config.FeatureLevel);
+				const int64 StartPos = Writer.Tell();
+				Writer << *ChunkAccessHeader;
+				ChunkAccessHeader->HeaderSize = Writer.Tell() - StartPos;
+				ChunkAccessHeader->StoredAs = bDataIsCompressed ? EChunkStorageFlags::Compressed : EChunkStorageFlags::RawData;
+				if (bDataIsEncrypted)
+				{
+					ChunkAccessHeader->StoredAs |= EChunkStorageFlags::Encrypted;
+					ChunkAccessHeader->EncryptionSecretId = Config.WriteEncryptionSecretId;
+					FMemory::Memcpy(ChunkAccessHeader->AESAuthTag.AuthTag, TempEncryptionAuthTag.GetData(), AES256_GCM_AuthTagSizeInBytes);
+				}
+				ChunkAccessHeader->DataSizeCompressed = CompressedSize;
+				ChunkAccessHeader->DataSizeUncompressed = ChunkAccessHeader->DataSizeUncompressed;
+				// Make sure we at least have rolling hash.
+				if ((ChunkAccessHeader->HashType & EChunkHashFlags::RollingPoly64) == EChunkHashFlags::None)
+				{
+					ChunkAccessHeader->HashType = EChunkHashFlags::RollingPoly64;
+				}
+
+				// Write out files.
+				Writer.Seek(StartPos);
+				Writer << *ChunkAccessHeader;
+				Writer.Serialize(ChunkDataSource, CompressedSize);
+				if (Writer.IsError())
+				{
+					SaveResult = EChunkSaveResult::SerializationError;
+				}
 			}
 			ChunkDataAccess->ReleaseDataLock();
 			return SaveResult;
@@ -963,12 +1130,15 @@ namespace BuildPatchServices
 
 	private:
 		IFileSystem* const FileSystem;
-		const EFeatureLevel FeatureLevel;
+		ICrypto* const Crypto;
+		const FChunkDataSerializationConfig Config;
+		const TArray<uint8>* const WriteEncryptionSecretKey;
+		// UE5 MERGE TODO: ICloudChunkSourceStat* CloudChunkSourceStat;
 	};
 
-	IChunkDataSerialization* FChunkDataSerializationFactory::Create(IFileSystem* FileSystem, EFeatureLevel FeatureLevel)
+	IChunkDataSerialization* FChunkDataSerializationFactory::Create(IFileSystem* FileSystem, ICrypto* Crypto, FChunkDataSerializationConfig Config)
 	{
 		check(FileSystem != nullptr);
-		return new FChunkDataSerialization(FileSystem, FeatureLevel);
+		return new FChunkDataSerialization(FileSystem, Crypto, Config);
 	}
 }
diff --git a/Engine/Source/Runtime/Online/BuildPatchServices/Private/Data/ChunkData.h b/Engine/Source/Runtime/Online/BuildPatchServices/Private/Data/ChunkData.h
--- a/Engine/Source/Runtime/Online/BuildPatchServices/Private/Data/ChunkData.h
+++ b/Engine/Source/Runtime/Online/BuildPatchServices/Private/Data/ChunkData.h
@@ -5,11 +5,14 @@
 #include "Misc/SecureHash.h"
 #include "Misc/EnumClassFlags.h"
 #include "Misc/Guid.h"
+
+#include "Common/Crypto.h"
 #include "BuildPatchFeatureLevel.h"
 
 namespace BuildPatchServices
 {
 	class IFileSystem;
+	class ICloudChunkSourceStat;
 
 	// Constant for the legacy fixed chunk window size, which was 1MiB.
 	static const uint32 LegacyFixedChunkWindow = 1024 * 1024;
@@ -19,13 +22,13 @@ namespace BuildPatchServices
 	 */
 	enum class EChunkStorageFlags : uint8
 	{
-		None = 0x00,
+		RawData = 0x00,
 
 		// Flag for compressed data.
 		Compressed = 0x01,
 
 		// Flag for encrypted. If also compressed, decrypt first. Encryption will ruin compressibility.
-		Encrypted = 0x02,
+		Encrypted = 0x02
 	};
 	ENUM_CLASS_FLAGS(EChunkStorageFlags);
 
@@ -41,6 +44,9 @@ namespace BuildPatchServices
 
 		// Flag for FSHA1 class used, stored in SHAHash on header.
 		Sha1 = 0x02,
+
+		// Flag for AES AuthTag being present in header header.
+		AesAuthTag = 0x02,
 	};
 	ENUM_CLASS_FLAGS(EChunkHashFlags);
 
@@ -81,6 +87,13 @@ namespace BuildPatchServices
 		// The size of the file is too big to handle 
 		FileSizeTooBig,
 
+		// If the decryption key necessary has not been provided, this result is given.
+		// It is up to the calling code to tell whether this is acceptable (expected preloading for example).
+		MissingSecretKey,
+
+		// The data was saved encrypted, and we had the secret key, but decryption failed.
+		DecryptionFailure,
+
 		// The operation was aborted.
 		Aborted
 	};
@@ -104,7 +117,15 @@ namespace BuildPatchServices
 		BadArchive,
 
 		// There was a serialization problem when writing to the chunk file.
-		SerializationError
+		SerializationError,
+
+		// If we were asked to encrypted saved out chunks but did not get the key, we cannot save at all.
+		MissingEncryptionKey,
+
+		// If the crypto interface returns an error.
+		EncryptionFail,
+
+		InvalidOrMax
 	};
 
 	/**
@@ -131,6 +152,8 @@ namespace BuildPatchServices
 		uint32 HeaderSize;
 		// The GUID for this data.
 		FGuid Guid;
+		// An ID that identifies the encryption secret required to decrypt this chunk. This is blank if the chunk is not encrypted.
+		FGuid EncryptionSecretId;
 		// The size of this data compressed.
 		uint32 DataSizeCompressed;
 		// The size of this data uncompressed.
@@ -143,6 +166,8 @@ namespace BuildPatchServices
 		uint64 RollingHash;
 		// The FSHA hashed value for this chunk data.
 		FSHAHash SHAHash;
+		// The 16 byte AuthTag space used for AES encryption, which is an additional integrity check value.
+		FAESAuthTag AESAuthTag;
 	};
 
 	// Some helper constants for dealing with chunks that are full of one single byte, usually padding.
@@ -217,6 +242,8 @@ namespace BuildPatchServices
 		uint32 Offset;
 		// The size of this part.
 		uint32 Size;
+		// Determine whether the chunk is from delta file or manifest (default).
+		bool bChunkFromDelta = false;
 	};
 
 	/**
@@ -238,19 +265,26 @@ namespace BuildPatchServices
 	 */
 	struct FChunkInfo
 	{
-		FChunkInfo();
 		// The GUID for this data.
 		FGuid Guid;
 		// The FRollingHash hashed value for this chunk data.
-		uint64 Hash;
+		uint64 Hash = 0;
 		// The FSHA hashed value for this chunk data.
 		FSHAHash ShaHash;
 		// The group number this chunk divides into.
-		uint8 GroupNumber;
-		// The window size for this chunk.
-		uint32 WindowSize;
+		uint8 GroupNumber = 0;
+		// The size of this data compressed.
+		uint32 DataSizeCompressed = 0;
+		// The size of this data uncompressed (also known as window size in older algorithms).
+		uint32 DataSizeUncompressed = 1048576;
 		// The file download size for this chunk.
-		int64 FileSize;
+		int64 FileSize = 0;
+		// The ID of the encryption secret key used by this chunk.
+		FGuid EncryptionSecretId;
+		// The 16 byte AuthTag space used for AES encryption, which is an additional integrity check value.
+		FAESAuthTag AESAuthTag;
+		// Determine whether the chunk is from delta file or manifest (default).
+		bool bChunkFromDelta = false;
 	};
 
 	/**
@@ -383,32 +417,42 @@ namespace BuildPatchServices
 		 */
 		virtual IChunkDataAccess* LoadFromArchive(FArchive& Archive, EChunkLoadResult& OutLoadResult) const = 0;
 		
+		/**
+		 * Loads a chunk from an archive and provide the "task" to decrypt and decode loaded data.
+		 * @param Archive           The archive.
+		 * @param OutLoadResult     Receives the result, indicating the error reason if return value is nullptr.
+		 * @return the function, returning a pair: 
+		 *  - code indicating the error reason;
+		 *  - ptr to an allocated IChunkDataAccess holding the data, nullptr if failed to async load.
+		 */
+		virtual TFunction<TPair<EChunkLoadResult, IChunkDataAccess*>()> LoadFromArchiveAsync(FArchive& Archive) const = 0;
+
 		/**
 		 * Saves a chunk to a file on disk or network.
 		 * @param Filename          The full file path to the file.
-		 * @param ChunkDataAccess   Ptr to the chunk data to save.
+		 * @param ChunkDataAccess   Ptr to the chunk data to save. The header will be updated with relevant changes.
 		 * @return the result, EChunkSaveResult::Success if successful.
 		 */
-		virtual EChunkSaveResult SaveToFile(const FString& Filename, const IChunkDataAccess* ChunkDataAccess) const = 0;
+		virtual EChunkSaveResult SaveToFile(const FString& Filename, IChunkDataAccess* ChunkDataAccess) const = 0;
 
 		/**
 		 * Saves a chunk to memory.
 		 * @param Memory            The memory array.
-		 * @param ChunkDataAccess   Ptr to the chunk data to save.
+		 * @param ChunkDataAccess   Ptr to the chunk data to save. The header will be updated with relevant changes.
 		 * @return the result, EChunkSaveResult::Success if successful.
 		 */
-		virtual EChunkSaveResult SaveToMemory(TArray<uint8>& Memory, const IChunkDataAccess* ChunkDataAccess) const = 0;
+		virtual EChunkSaveResult SaveToMemory(TArray<uint8>& Memory, IChunkDataAccess* ChunkDataAccess) const = 0;
 
 		/**
 		 * Saves a chunk to an archive.
 		 * @param Archive           The archive.
-		 * @param ChunkDataAccess   Ptr to the chunk data to save.
+		 * @param ChunkDataAccess   Ptr to the chunk data to save. The header will be updated with relevant changes.
 		 * @return the result, EChunkSaveResult::Success if successful.
 		 */
-		virtual EChunkSaveResult SaveToArchive(FArchive& Archive, const IChunkDataAccess* ChunkDataAccess) const = 0;
+		virtual EChunkSaveResult SaveToArchive(FArchive& Archive, IChunkDataAccess* ChunkDataAccess) const = 0;
 
 		// In performance sensitive situations, can request the serializer to save uncompressed if implemented
-		virtual EChunkSaveResult SaveToArchiveUncompressed(FArchive& Archive, const IChunkDataAccess* ChunkDataAccess) const
+		virtual EChunkSaveResult SaveToArchiveUncompressed(FArchive& Archive, IChunkDataAccess* ChunkDataAccess) const
 		{
 			return SaveToArchive(Archive, ChunkDataAccess);
 		}
@@ -418,14 +462,23 @@ namespace BuildPatchServices
 		virtual bool ValidateAndRead(FArchive& InArchive, FMutableMemoryView InDestinationBuffer, FChunkHeader& OutHeader, FUniqueBuffer& OutCompressedBuffer) const = 0;
 		
 		// Take the output of ValidateAndRead and process it in to an actual output chunk.
-		virtual bool DecompressValidatedRead(const FChunkHeader& InHeader, FMutableMemoryView InDestionationBuffer, const FUniqueBuffer& InCompressedBuffer) const = 0;
+		virtual bool DecompressAndDecryptValidatedRead(const FChunkHeader& InHeader, FMutableMemoryView InDestionationBuffer, const FUniqueBuffer& InCompressedBuffer) const = 0;
 
 		/**
 		 * Injects an SHA hash for the data into the structure of a serialized chunk.
 		 * @param Memory            The memory array containing the serialized chunk.
 		 * @param ShaHashData       The SHA hash to inject.
+		 * @return true if ShaHashData was successfully injected into memory, false otherwise.
 		 */
-		virtual void InjectShaToChunkData(TArray<uint8>& Memory, const FSHAHash& ShaHashData) const = 0;
+		virtual bool InjectShaToChunkData(TArray<uint8>& Memory, const FSHAHash& ShaHashData) const = 0;
+	};
+
+	struct FChunkDataSerializationConfig
+	{
+		EFeatureLevel FeatureLevel = EFeatureLevel::Latest;
+		TMap<FGuid,TArray<uint8>> EncryptionSecrets;
+		FGuid WriteEncryptionSecretId;
+		// UE5 MERGE TODO: ICloudChunkSourceStat* CloudChunkSourceStat;
 	};
 
 	/**
@@ -434,6 +487,6 @@ namespace BuildPatchServices
 	class FChunkDataSerializationFactory
 	{
 	public:
-		static IChunkDataSerialization* Create(IFileSystem* FileSystem, EFeatureLevel FeatureLevel = EFeatureLevel::Latest);
+		static IChunkDataSerialization* Create(IFileSystem* FileSystem, ICrypto* Crypto, FChunkDataSerializationConfig Config);
 	};
-}
+}
\ No newline at end of file
diff --git a/Engine/Source/Runtime/Online/BuildPatchServices/Private/Data/ManifestData.cpp b/Engine/Source/Runtime/Online/BuildPatchServices/Private/Data/ManifestData.cpp
--- a/Engine/Source/Runtime/Online/BuildPatchServices/Private/Data/ManifestData.cpp
+++ b/Engine/Source/Runtime/Online/BuildPatchServices/Private/Data/ManifestData.cpp
@@ -1,7 +1,9 @@
 // Copyright Epic Games, Inc. All Rights Reserved.
 
 #include "Data/ManifestData.h"
+#include "Misc/Base64.h"
 #include "Misc/Compression.h"
+#include "Misc/ConfigCacheIni.h"
 #include "Misc/EnumClassFlags.h"
 #include "Serialization/Archive.h"
 #include "Serialization/MemoryWriter.h"
@@ -31,10 +33,12 @@ namespace BuildPatchServices
 		// This remained the same all up to including EFeatureLevel::StoresPrerequisiteIds.
 		37, 37, 37, 37, 37, 37, 37, 37, 37, 37, 37, 37, 37, 37,
 		// EFeatureLevel::StoredAsBinaryData is 41B, (296b Original, 32b Version).
-		// This remained the same all up to including EFeatureLevel::UsesBuildTimeGeneratedBuildId.
-		41, 41, 41, 41, 41
+		// This remained the same all up to including EFeatureLevel::StoresUninstallActions.
+		41, 41, 41, 41, 41, 41, 41, 41,
+		// EFeatureLevel::ChunkEncryptionSupport added a 128b SecretId GUID, plus 128b AuthTag
+		73, 73, 73
 	};
-	static_assert((int32)EFeatureLevel::Latest == 18, "Please adjust ManifestHeaderVersionSizes values accordingly.");
+	static_assert((int32)EFeatureLevel::Latest == 24, "Please adjust ManifestHeaderVersionSizes values accordingly.");
 
 	/**
 	 * Enum which describes the FManifestMeta data version.
@@ -43,6 +47,8 @@ namespace BuildPatchServices
 	{
 		Original = 0,
 		SerialisesBuildId,
+		SerialisesUnistallActions,
+		SerialisesManifestEncryptionSecretId,
 
 		// Always after the latest version, signifies the latest version plus 1 to allow initialization simplicity.
 		LatestPlusOne,
@@ -55,6 +61,9 @@ namespace BuildPatchServices
 	enum class EChunkDataListVersion : uint8
 	{
 		Original = 0,
+		SerialisesEncryptionSecretId,
+		SerialisesCompressesDataSize,
+		SerialisesAESAuthTag,
 
 		// Always after the latest version, signifies the latest version plus 1 to allow initialization simplicity.
 		LatestPlusOne,
@@ -65,6 +74,20 @@ namespace BuildPatchServices
 	 * Enum which describes the FFileManifestList data version.
 	 */
 	enum class EFileManifestListVersion : uint8
+	{
+		Original = 0,
+		HasMD5AndMIMEType,
+		HasSHA256,
+
+		// Always after the latest version, signifies the latest version plus 1 to allow initialization simplicity.
+		LatestPlusOne,
+		Latest = (LatestPlusOne - 1)
+	};
+
+	/**
+	 * Enum which describes the EEncryptedDataVersion data version.
+	 */
+	enum class EEncryptedDataVersion : uint8
 	{
 		Original = 0,
 
@@ -78,20 +101,22 @@ namespace BuildPatchServices
 		const TCHAR* GetChunkSubdir(EFeatureLevel FeatureLevel)
 		{
 			return FeatureLevel < EFeatureLevel::DataFileRenames ? TEXT("Chunks")
-				: FeatureLevel < EFeatureLevel::ChunkCompressionSupport ? TEXT("ChunksV2")
-				: FeatureLevel < EFeatureLevel::VariableSizeChunksWithoutWindowSizeChunkInfo ? TEXT("ChunksV3")
-				: TEXT("ChunksV4");
+				 : FeatureLevel < EFeatureLevel::ChunkCompressionSupport ? TEXT("ChunksV2")
+				 : FeatureLevel < EFeatureLevel::VariableSizeChunksWithoutWindowSizeChunkInfo ? TEXT("ChunksV3")
+				 : FeatureLevel < EFeatureLevel::ChunkEncryptionSupport ? TEXT("ChunksV4")
+				 : TEXT("ChunksV5");
 		}
 
 		const TCHAR* GetFileSubdir(EFeatureLevel FeatureLevel)
 		{
 			return FeatureLevel < EFeatureLevel::DataFileRenames ? TEXT("Files")
-				: FeatureLevel < EFeatureLevel::StoresChunkDataShaHashes ? TEXT("FilesV2")
-				: TEXT("FilesV3");
+				 : FeatureLevel < EFeatureLevel::StoresChunkDataShaHashes ? TEXT("FilesV2")
+				 : TEXT("FilesV3");
 		}
 
 		EManifestMetaVersion FeatureLevelToManifestMetaVersion(EFeatureLevel FeatureLevel)
 		{
+			static_assert((uint32)EFeatureLevel::Latest == 24, "Please adjust the switch below for new feature levels.");
 			switch (FeatureLevel)
 			{
 			case BuildPatchServices::EFeatureLevel::Original:
@@ -114,13 +139,131 @@ namespace BuildPatchServices
 			case BuildPatchServices::EFeatureLevel::UsesRuntimeGeneratedBuildId:
 				return EManifestMetaVersion::Original;
 			case BuildPatchServices::EFeatureLevel::UsesBuildTimeGeneratedBuildId:
+			case BuildPatchServices::EFeatureLevel::StoresFileMD5HashesAndMIMEType:
+			case BuildPatchServices::EFeatureLevel::StoresFileSHA256Hashes:
 				return EManifestMetaVersion::SerialisesBuildId;
+			case BuildPatchServices::EFeatureLevel::StoresUninstallActions:
+				return EManifestMetaVersion::SerialisesUnistallActions;
+			case BuildPatchServices::EFeatureLevel::ChunkEncryptionSupport:
+			case BuildPatchServices::EFeatureLevel::ChunksStoredBySecret:
+			case BuildPatchServices::EFeatureLevel::ManifestEncryptionSupport:
+				return EManifestMetaVersion::SerialisesManifestEncryptionSecretId;
 			}
-			checkf(false, TEXT("Unhandled FeatureLevel %s"), FeatureLevelToString(FeatureLevel));
+			checkf(false, TEXT("Unhandled FeatureLevel %s"), LexToString(FeatureLevel));
 			return EManifestMetaVersion::Latest;
 		}
+
+		EChunkDataListVersion FeatureLevelToChunkDataListVersion(EFeatureLevel FeatureLevel)
+		{
+			static_assert((uint32)EFeatureLevel::Latest == 24, "Please adjust the switch below for new feature levels.");
+			switch (FeatureLevel)
+			{
+			case BuildPatchServices::EFeatureLevel::Original:
+			case BuildPatchServices::EFeatureLevel::CustomFields:
+			case BuildPatchServices::EFeatureLevel::StartStoringVersion:
+			case BuildPatchServices::EFeatureLevel::DataFileRenames:
+			case BuildPatchServices::EFeatureLevel::StoresIfChunkOrFileData:
+			case BuildPatchServices::EFeatureLevel::StoresDataGroupNumbers:
+			case BuildPatchServices::EFeatureLevel::ChunkCompressionSupport:
+			case BuildPatchServices::EFeatureLevel::StoresPrerequisitesInfo:
+			case BuildPatchServices::EFeatureLevel::StoresChunkFileSizes:
+			case BuildPatchServices::EFeatureLevel::StoredAsCompressedUClass:
+			case BuildPatchServices::EFeatureLevel::UNUSED_0:
+			case BuildPatchServices::EFeatureLevel::UNUSED_1:
+			case BuildPatchServices::EFeatureLevel::StoresChunkDataShaHashes:
+			case BuildPatchServices::EFeatureLevel::StoresPrerequisiteIds:
+			case BuildPatchServices::EFeatureLevel::StoredAsBinaryData:
+			case BuildPatchServices::EFeatureLevel::VariableSizeChunksWithoutWindowSizeChunkInfo:
+			case BuildPatchServices::EFeatureLevel::VariableSizeChunks:
+			case BuildPatchServices::EFeatureLevel::UsesRuntimeGeneratedBuildId:
+			case BuildPatchServices::EFeatureLevel::UsesBuildTimeGeneratedBuildId:
+			case BuildPatchServices::EFeatureLevel::StoresFileMD5HashesAndMIMEType:
+			case BuildPatchServices::EFeatureLevel::StoresFileSHA256Hashes:
+			case BuildPatchServices::EFeatureLevel::StoresUninstallActions:
+				return EChunkDataListVersion::Original;
+			case BuildPatchServices::EFeatureLevel::ChunkEncryptionSupport:
+			case BuildPatchServices::EFeatureLevel::ChunksStoredBySecret:
+			case BuildPatchServices::EFeatureLevel::ManifestEncryptionSupport:
+				return EChunkDataListVersion::SerialisesAESAuthTag;
+			}
+			checkf(false, TEXT("Unhandled FeatureLevel %s"), LexToString(FeatureLevel));
+			return EChunkDataListVersion::Latest;
+		}
+
+		EFileManifestListVersion FeatureLevelToFileManifestListVersion(EFeatureLevel FeatureLevel)
+		{
+			static_assert((uint32)EFeatureLevel::Latest == 24, "Please adjust the switch below for new feature levels.");
+			switch (FeatureLevel)
+			{
+				case BuildPatchServices::EFeatureLevel::Original:
+				case BuildPatchServices::EFeatureLevel::CustomFields:
+				case BuildPatchServices::EFeatureLevel::StartStoringVersion:
+				case BuildPatchServices::EFeatureLevel::DataFileRenames:
+				case BuildPatchServices::EFeatureLevel::StoresIfChunkOrFileData:
+				case BuildPatchServices::EFeatureLevel::StoresDataGroupNumbers:
+				case BuildPatchServices::EFeatureLevel::ChunkCompressionSupport:
+				case BuildPatchServices::EFeatureLevel::StoresPrerequisitesInfo:
+				case BuildPatchServices::EFeatureLevel::StoresChunkFileSizes:
+				case BuildPatchServices::EFeatureLevel::StoredAsCompressedUClass:
+				case BuildPatchServices::EFeatureLevel::UNUSED_0:
+				case BuildPatchServices::EFeatureLevel::UNUSED_1:
+				case BuildPatchServices::EFeatureLevel::StoresChunkDataShaHashes:
+				case BuildPatchServices::EFeatureLevel::StoresPrerequisiteIds:
+				case BuildPatchServices::EFeatureLevel::StoredAsBinaryData:
+				case BuildPatchServices::EFeatureLevel::VariableSizeChunksWithoutWindowSizeChunkInfo:
+				case BuildPatchServices::EFeatureLevel::VariableSizeChunks:
+				case BuildPatchServices::EFeatureLevel::UsesRuntimeGeneratedBuildId:
+				case BuildPatchServices::EFeatureLevel::UsesBuildTimeGeneratedBuildId:
+					return EFileManifestListVersion::Original;
+				case BuildPatchServices::EFeatureLevel::StoresFileMD5HashesAndMIMEType:
+					return EFileManifestListVersion::HasMD5AndMIMEType;
+				case BuildPatchServices::EFeatureLevel::StoresFileSHA256Hashes:
+				case BuildPatchServices::EFeatureLevel::StoresUninstallActions:
+				case BuildPatchServices::EFeatureLevel::ChunkEncryptionSupport:
+				case BuildPatchServices::EFeatureLevel::ChunksStoredBySecret:
+				case BuildPatchServices::EFeatureLevel::ManifestEncryptionSupport:
+					return EFileManifestListVersion::HasSHA256;
+			}
+			checkf(false, TEXT("Unhandled FeatureLevel %s"), LexToString(FeatureLevel));
+			return EFileManifestListVersion::Latest;
+		}
+
+		EEncryptedDataVersion FeatureLevelToEncryptedDataVersion(EFeatureLevel FeatureLevel)
+		{
+			static_assert((uint32)EFeatureLevel::Latest == 24, "Please adjust the switch below for new feature levels.");
+			switch (FeatureLevel)
+			{
+				case BuildPatchServices::EFeatureLevel::Original:
+				case BuildPatchServices::EFeatureLevel::CustomFields:
+				case BuildPatchServices::EFeatureLevel::StartStoringVersion:
+				case BuildPatchServices::EFeatureLevel::DataFileRenames:
+				case BuildPatchServices::EFeatureLevel::StoresIfChunkOrFileData:
+				case BuildPatchServices::EFeatureLevel::StoresDataGroupNumbers:
+				case BuildPatchServices::EFeatureLevel::ChunkCompressionSupport:
+				case BuildPatchServices::EFeatureLevel::StoresPrerequisitesInfo:
+				case BuildPatchServices::EFeatureLevel::StoresChunkFileSizes:
+				case BuildPatchServices::EFeatureLevel::StoredAsCompressedUClass:
+				case BuildPatchServices::EFeatureLevel::UNUSED_0:
+				case BuildPatchServices::EFeatureLevel::UNUSED_1:
+				case BuildPatchServices::EFeatureLevel::StoresChunkDataShaHashes:
+				case BuildPatchServices::EFeatureLevel::StoresPrerequisiteIds:
+				case BuildPatchServices::EFeatureLevel::StoredAsBinaryData:
+				case BuildPatchServices::EFeatureLevel::VariableSizeChunksWithoutWindowSizeChunkInfo:
+				case BuildPatchServices::EFeatureLevel::VariableSizeChunks:
+				case BuildPatchServices::EFeatureLevel::UsesRuntimeGeneratedBuildId:
+				case BuildPatchServices::EFeatureLevel::UsesBuildTimeGeneratedBuildId:
+				case BuildPatchServices::EFeatureLevel::StoresFileMD5HashesAndMIMEType:
+				case BuildPatchServices::EFeatureLevel::StoresFileSHA256Hashes:
+				case BuildPatchServices::EFeatureLevel::StoresUninstallActions:
+				case BuildPatchServices::EFeatureLevel::ChunkEncryptionSupport:
+				case BuildPatchServices::EFeatureLevel::ChunksStoredBySecret:
+				case BuildPatchServices::EFeatureLevel::ManifestEncryptionSupport:
+					return EEncryptedDataVersion::Original;
+			}
+			checkf(false, TEXT("Unhandled FeatureLevel %s"), LexToString(FeatureLevel));
+			return EEncryptedDataVersion::Latest;
+		}
 	}
-	static_assert((uint32)EFeatureLevel::Latest == 18, "Please adjust ManifestVersionHelpers::FeatureLevelToManifestMetaVersion for new feature levels.");
 
 	namespace ManifestDataHelpers
 	{
@@ -141,6 +284,105 @@ namespace BuildPatchServices
 				return TUniquePtr<FArchive>(new FMemoryWriter(Memory));
 			}
 		}
+
+		EManifestStorageFlags CompressManifestMemory(TArray<uint8>& Data)
+		{
+			const uint32 DataSizeUncompressed = Data.Num();
+			// Compression format selection - we only have one right now.
+			const FName CompressionFormat = NAME_Zlib;
+			const ECompressionFlags CompressionFlags = ECompressionFlags::COMPRESS_BiasMemory;
+
+			// Sanely clamp compression effort.
+			int MinCompressionBytes = 64 * 1024;
+			int MinCompressionPct = 10;
+			GConfig->GetInt(TEXT("BuildPatchTool"), TEXT("MinCompressionBytes"), MinCompressionBytes, GEngineIni);
+			GConfig->GetInt(TEXT("BuildPatchTool"), TEXT("MinCompressionPct"), MinCompressionPct, GEngineIni);
+
+			// Perform compression
+			TArray<uint8> TempCompressed;
+			uint32 DataSizeCompressed = FCompression::CompressMemoryBound(CompressionFormat, DataSizeUncompressed, CompressionFlags);
+			TempCompressed.AddUninitialized(DataSizeCompressed);
+			const bool bDataIsCompressed = FCompression::CompressMemoryIfWorthDecompressing(
+				CompressionFormat,
+				MinCompressionBytes, // bytes is the minimum worth compressing.
+				MinCompressionPct, // percent is the minimum worth compressing.
+				TempCompressed.GetData(),
+				(int32&)DataSizeCompressed,
+				Data.GetData(),
+				Data.Num(),
+				CompressionFlags);
+			EManifestStorageFlags StoredAs;
+			if (bDataIsCompressed)
+			{
+				TempCompressed.SetNum(DataSizeCompressed, EAllowShrinking::No);
+				Data = MoveTemp(TempCompressed);
+				StoredAs = EManifestStorageFlags::Compressed;
+			}
+			else
+			{
+				DataSizeCompressed = Data.Num();
+				StoredAs = EManifestStorageFlags::RawData;
+			}
+			return StoredAs;
+		}
+
+		bool UncompressManifestMemory(TArray<uint8>& Data, uint32 DataSizeUncompressed)
+		{
+			// Compression format selection - we only have one right now.
+			const FName CompressionFormat = NAME_Zlib;
+			const ECompressionFlags CompressionFlags = ECompressionFlags::COMPRESS_BiasMemory;
+
+			TArray<uint8> CompressedData = MoveTemp(Data);
+			Data.AddUninitialized(DataSizeUncompressed);
+			return FCompression::UncompressMemory(
+				CompressionFormat,
+				Data.GetData(),
+				Data.Num(),
+				CompressedData.GetData(),
+				CompressedData.Num(),
+				CompressionFlags);
+		}
+
+		FString ObfuscateString(const FString& String)
+		{
+			if (String.Len() > 0)
+			{
+				FSHA1 Sha;
+				FSHAHash Hash;
+				FTCHARToUTF8 UTF8String(*String);
+				FSHA1::HashBuffer((const uint8*)UTF8String.Get(), UTF8String.Length(), Hash.Hash);
+				FString ObfuscatedString;
+				FBuildPatchUtils::SHAToBase32(Hash, ObfuscatedString);
+				return ObfuscatedString;
+			}
+			return String;
+		}
+
+		TArray<FString> ObfuscateStrings(const TArray<FString>& Strings)
+		{
+			TArray<FString> Result = Strings;
+			for (FString& Val : Result)
+			{
+				Val = ObfuscateString(Val);
+			}
+			return Result;
+		}
+
+		TSet<FString> ObfuscateStrings(const TSet<FString>& Strings)
+		{
+			TSet<FString> Result = Strings;
+			for (FString& Val : Result)
+			{
+				Val = ObfuscateString(Val);
+			}
+			return Result;
+		}
+	}
+
+	FArchive& operator<<(FArchive& Ar, FSHA256Signature& Hash)
+	{
+		Ar.Serialize(&Hash.Signature, sizeof(Hash.Signature));
+		return Ar;
 	}
 
 	/* FManifestHeader - The header for a compressed/encoded manifest file.
@@ -151,7 +393,7 @@ namespace BuildPatchServices
 		, HeaderSize(0)
 		, DataSizeCompressed(0)
 		, DataSizeUncompressed(0)
-		, StoredAs(EManifestStorageFlags::None)
+		, StoredAs(EManifestStorageFlags::RawData)
 		, SHAHash()
 	{
 	}
@@ -198,6 +440,22 @@ namespace BuildPatchServices
 			{
 				Header.Version = EFeatureLevel::StoredAsCompressedUClass;
 			}
+
+			// With encryption support, came the encryption secretID in the manifest header
+			// We also stored the AuthTag for the manifest.
+			if (!Ar.IsError() && Header.Version >= EFeatureLevel::ChunkEncryptionSupport)
+			{
+				Ar << Header.EncryptionSecretId;
+				Ar.Serialize(Header.EncryptionAuthTag, AES256_GCM_AuthTagSizeInBytes);
+				ExpectedSerializedBytes = ManifestHeaderVersionSizes[(int32)EFeatureLevel::ChunkEncryptionSupport];
+			}
+
+			//// Here we would check for later data versions to serialise additional values.
+			//if (!Ar.IsError() && Header.Version >= EFeatureLevel::SomeShinyNewVersion)
+			//{
+			//	Ar << Header.SomeShinyNewVariable;
+			//	ExpectedSerializedBytes = ManifestHeaderVersionSizes[(int32)EFeatureLevel::SomeShinyNewVersion];
+			//}
 		}
 
 		// Make sure the expected number of bytes were serialized. In practice this will catch errors where type
@@ -282,6 +540,13 @@ namespace BuildPatchServices
 			Meta.BuildId = FBuildPatchUtils::GetBackwardsCompatibleBuildId(Meta);
 		}
 
+		// Serialise the Uninstall Actions.
+		if (!Ar.IsError() && DataVersion >= EManifestMetaVersion::SerialisesUnistallActions)
+		{
+			Ar << Meta.UninstallActionPath;
+			Ar << Meta.UninstallActionArgs;
+		}
+
 		//// Here we would check for later data versions to serialise additional values.
 		//if (!Ar.IsError() && DataVersion >= ManifestMetaVersion::SomeShinyNewVersion)
 		//{
@@ -307,6 +572,7 @@ namespace BuildPatchServices
 	*****************************************************************************/
 
 	FChunkDataList::FChunkDataList()
+		: FeatureLevel(BuildPatchServices::EFeatureLevel::Invalid)
 	{
 	}
 
@@ -320,7 +586,7 @@ namespace BuildPatchServices
 		// Serialise the data header type values.
 		const int64 StartPos = Ar.Tell();
 		uint32 DataSize = 0;
-		EChunkDataListVersion DataVersion = EChunkDataListVersion::Latest;
+		EChunkDataListVersion DataVersion = Ar.IsSaving() ? ManifestVersionHelpers::FeatureLevelToChunkDataListVersion(ChunkDataList.FeatureLevel) : EChunkDataListVersion::Latest;
 		int32 ElementCount = ChunkDataList.ChunkList.Num();
 		{
 			uint8 DataVersionInt = (uint8)DataVersion;
@@ -344,10 +610,24 @@ namespace BuildPatchServices
 			for (FChunkInfo& ChunkInfo : ChunkDataList.ChunkList) { Ar << ChunkInfo.Hash; }
 			for (FChunkInfo& ChunkInfo : ChunkDataList.ChunkList) { Ar << ChunkInfo.ShaHash; }
 			for (FChunkInfo& ChunkInfo : ChunkDataList.ChunkList) { Ar << ChunkInfo.GroupNumber; }
-			for (FChunkInfo& ChunkInfo : ChunkDataList.ChunkList) { Ar << ChunkInfo.WindowSize; }
+			for (FChunkInfo& ChunkInfo : ChunkDataList.ChunkList) { Ar << ChunkInfo.DataSizeUncompressed; }
 			for (FChunkInfo& ChunkInfo : ChunkDataList.ChunkList) { Ar << ChunkInfo.FileSize; }
 		}
 
+		// If we also have encryption, serialise these members.
+		if (!Ar.IsError() && DataVersion >= EChunkDataListVersion::SerialisesEncryptionSecretId)
+		{
+			for (FChunkInfo& ChunkInfo : ChunkDataList.ChunkList) { Ar << ChunkInfo.EncryptionSecretId; }
+		}
+		if (!Ar.IsError() && DataVersion >= EChunkDataListVersion::SerialisesCompressesDataSize)
+		{
+			for (FChunkInfo& ChunkInfo : ChunkDataList.ChunkList) { Ar << ChunkInfo.DataSizeCompressed; }
+		}
+		if (!Ar.IsError() && DataVersion >= EChunkDataListVersion::SerialisesAESAuthTag)
+		{
+			for (FChunkInfo& ChunkInfo : ChunkDataList.ChunkList) { Ar.Serialize(ChunkInfo.AESAuthTag.AuthTag, AES256_GCM_AuthTagSizeInBytes); }
+		}
+
 		//// Here we would check for later data versions to serialise additional values.
 		//if (!Ar.IsError() && DataVersion >= EChunkDataListVersion::SomeShinyNewVersion)
 		//{
@@ -382,13 +662,12 @@ namespace BuildPatchServices
 	*****************************************************************************/
 
 	FFileManifestList::FFileManifestList()
+		: FeatureLevel(BuildPatchServices::EFeatureLevel::Invalid)
 	{
 	}
 
 	void FFileManifestList::OnPostLoad()
 	{
-		Algo::SortBy(FileList, [](const FFileManifest& FileManifest) -> const FString& { return FileManifest.Filename; }, TLess<FString>());
-
 		for (FFileManifest& FileManifest : FileList)
 		{
 			FileManifest.FileSize = Algo::Accumulate<int64>(FileManifest.ChunkParts, 0, [](int64 Count, const FChunkPart& ChunkPart){ return Count + ChunkPart.Size; });
@@ -405,7 +684,7 @@ namespace BuildPatchServices
 		// Serialise the data header type values.
 		const int64 StartPos = Ar.Tell();
 		uint32 DataSize = 0;
-		EFileManifestListVersion DataVersion = EFileManifestListVersion::Latest;
+		EFileManifestListVersion DataVersion = Ar.IsSaving() ? ManifestVersionHelpers::FeatureLevelToFileManifestListVersion(FileDataList.FeatureLevel) : EFileManifestListVersion::Latest;
 		int32 ElementCount = FileDataList.FileList.Num();
 		{
 			uint8 DataVersionInt = (uint8)DataVersion;
@@ -424,7 +703,7 @@ namespace BuildPatchServices
 		{
 			for (FFileManifest& FileManifest : FileDataList.FileList) { Ar << FileManifest.Filename; }
 			for (FFileManifest& FileManifest : FileDataList.FileList) { Ar << FileManifest.SymlinkTarget; }
-			for (FFileManifest& FileManifest : FileDataList.FileList) { Ar << FileManifest.FileHash; }
+			for (FFileManifest& FileManifest : FileDataList.FileList) { Ar << FileManifest.SHA1Hash; }
 			for (FFileManifest& FileManifest : FileDataList.FileList)
 			{
 				uint8 FileMetaFlagsInt = (uint8)FileManifest.FileMetaFlags;
@@ -435,6 +714,19 @@ namespace BuildPatchServices
 			for (FFileManifest& FileManifest : FileDataList.FileList) { Ar << FileManifest.ChunkParts; }
 		}
 
+		// Serialise the ManifestMetaVersion::HasMD5AndMIMEType version variables.
+		if (!Ar.IsError() && DataVersion >= EFileManifestListVersion::HasMD5AndMIMEType)
+		{
+			for (FFileManifest& FileManifest : FileDataList.FileList) { Ar << FileManifest.MD5Hash; }
+			for (FFileManifest& FileManifest : FileDataList.FileList) { Ar << FileManifest.MIMEType; }
+		}
+
+		// Serialise the ManifestMetaVersion::HasSHA256 version variables.
+		if (!Ar.IsError() && DataVersion >= EFileManifestListVersion::HasSHA256)
+		{
+			for (FFileManifest& FileManifest : FileDataList.FileList) { Ar << FileManifest.SHA256Hash; }
+		}
+
 		//// Here we would check for later data versions to serialise additional values.
 		//if (!Ar.IsError() && DataVersion >= EFileManifestListVersion::SomeShinyNewVersion)
 		//{
@@ -466,6 +758,7 @@ namespace BuildPatchServices
 	*****************************************************************************/
 
 	FCustomFields::FCustomFields()
+		: FeatureLevel(BuildPatchServices::EFeatureLevel::Invalid)
 	{
 	}
 
@@ -536,6 +829,65 @@ namespace BuildPatchServices
 		return Ar;
 	}
 
+	/* FEncryptedData - The data implementation for manifest ciphertext.
+	*****************************************************************************/
+
+	FEncryptedData::FEncryptedData()
+		: FeatureLevel(BuildPatchServices::EFeatureLevel::Invalid)
+	{
+	}
+
+	FArchive& operator<<(FArchive& Ar, FEncryptedData& EncryptedData)
+	{
+		if (Ar.IsError())
+		{
+			return Ar;
+		}
+
+		// We didn't start serialising any encrypted manifest data until EFeatureLevel::ManifestEncryptionSupport, so early out if necessary.
+		if(EncryptedData.FeatureLevel < EFeatureLevel::ManifestEncryptionSupport)
+		{
+			return Ar;
+		}
+
+		// Serialise the data header type values.
+		const int64 StartPos = Ar.Tell();
+		uint32 DataSize = 0;
+		EEncryptedDataVersion DataVersion = Ar.IsSaving() ? ManifestVersionHelpers::FeatureLevelToEncryptedDataVersion(EncryptedData.FeatureLevel) : EEncryptedDataVersion::Latest;
+		{
+			uint8 DataVersionInt = (uint8)DataVersion;
+			Ar << DataSize;
+			Ar << DataVersionInt;
+			DataVersion = (EEncryptedDataVersion)DataVersionInt;
+		}
+
+		// Serialise the Original version, just the data itself.
+		if (!Ar.IsError() && DataVersion >= EEncryptedDataVersion::Original)
+		{
+			Ar << EncryptedData.Data;
+		}
+
+		//// Here we would check for later data versions to serialise additional values.
+		//if (!Ar.IsError() && DataVersion >= EEncryptedDataVersion::SomeShinyNewVersion)
+		//{
+		//	Ar << Meta.SomeShinyNewVariable;
+		//}
+
+		// If saving, we need to go back and set the data size.
+		if (!Ar.IsError() && Ar.IsSaving())
+		{
+			const int64 EndPos = Ar.Tell();
+			DataSize = EndPos - StartPos;
+			Ar.Seek(StartPos);
+			Ar << DataSize;
+			Ar.Seek(EndPos);
+		}
+
+		// We must always make sure to seek the archive to the correct end location.
+		Ar.Seek(StartPos + DataSize);
+		return Ar;
+	}
+
 	/* FManifestData - The public interface to load/saving manifest files.
 	*****************************************************************************/
 	void FManifestData::Init()
@@ -561,6 +913,7 @@ namespace BuildPatchServices
 
 	bool FManifestData::Serialize(FArchive& Ar, FBuildPatchAppManifest& AppManifest, BuildPatchServices::EFeatureLevel SaveFormat)
 	{
+		const EFeatureLevel OriginalFormat = AppManifest.GetFeatureLevel();
 		if (Ar.IsError())
 		{
 			return false;
@@ -582,6 +935,14 @@ namespace BuildPatchServices
 			{
 				Ar << Header;
 				bSuccess = !Ar.IsError();
+				AppManifest.EncryptionSecretId = Header.EncryptionSecretId;
+				FMemory::Memcpy(AppManifest.EncryptionAuthTag.AuthTag, Header.EncryptionAuthTag, AES256_GCM_AuthTagSizeInBytes);
+			}
+			// Otherwise set header values for saving
+			else
+			{
+				Header.EncryptionSecretId = AppManifest.EncryptionSecretId;
+				FMemory::Memcpy(Header.EncryptionAuthTag, AppManifest.EncryptionAuthTag.AuthTag, AES256_GCM_AuthTagSizeInBytes);
 			}
 
 			// If we are loading an old format, defer to the old code!
@@ -598,15 +959,11 @@ namespace BuildPatchServices
 			}
 			else
 			{
-				// Compression format selection - we only have one right now.
-				const FName CompressionFormat = NAME_Zlib;
-				const ECompressionFlags CompressionFlags = ECompressionFlags::COMPRESS_BiasMemory;
-				// Yay shiny new format!
 				TArray<uint8> ManifestRawData;
 				// Fill the array with loaded data.
 				if (bSuccess && Ar.IsLoading())
 				{
-					// DataSizeCompressed always equals the size of the data following the header.
+					// DataSizeCompressed always equals the size of the data following the header, even if not compressed.
 					ManifestRawData.AddUninitialized(Header.DataSizeCompressed);
 					Ar.Serialize(ManifestRawData.GetData(), Header.DataSizeCompressed);
 					bSuccess = !Ar.IsError();
@@ -614,15 +971,7 @@ namespace BuildPatchServices
 				// Uncompress from input archive.
 				if (bSuccess && Ar.IsLoading() && EnumHasAllFlags(Header.StoredAs, EManifestStorageFlags::Compressed))
 				{
-					TArray<uint8> CompressedData = MoveTemp(ManifestRawData);
-					ManifestRawData.AddUninitialized(Header.DataSizeUncompressed);
-					bSuccess = FCompression::UncompressMemory(
-						CompressionFormat,
-						ManifestRawData.GetData(),
-						ManifestRawData.Num(),
-						CompressedData.GetData(),
-						CompressedData.Num(),
-						CompressionFlags);
+					bSuccess = ManifestDataHelpers::UncompressManifestMemory(ManifestRawData, Header.DataSizeUncompressed);
 				}
 				// If loading, check the raw data SHA
 				if (bSuccess && Ar.IsLoading())
@@ -633,20 +982,62 @@ namespace BuildPatchServices
 				}
 				if (bSuccess)
 				{
+					// If we are saving, set the save format on each complex serialisable member
+					if (Ar.IsSaving())
+					{
+						AppManifest.ManifestMeta.FeatureLevel = SaveFormat;
+						AppManifest.ChunkDataList.FeatureLevel = SaveFormat;
+						AppManifest.FileManifestList.FeatureLevel = SaveFormat;
+						AppManifest.CustomFields.FeatureLevel = SaveFormat;
+						AppManifest.EncryptedData.FeatureLevel = SaveFormat;
+					}
+
 					// Create the directional interface to the raw data array.
 					TUniquePtr<FArchive> RawAr = ManifestDataHelpers::CreateMemoryArchive(Ar.IsLoading(), ManifestRawData);
 					// Serialise each of the manifest's data members.
 					*RawAr << AppManifest.ManifestMeta;
+
+					// If we are loading, set the load format on each complex serialisable member, after getting it from the meta block.
+					if (Ar.IsLoading())
+					{
+						AppManifest.ChunkDataList.FeatureLevel = AppManifest.ManifestMeta.FeatureLevel;
+						AppManifest.FileManifestList.FeatureLevel = AppManifest.ManifestMeta.FeatureLevel;
+						AppManifest.CustomFields.FeatureLevel = AppManifest.ManifestMeta.FeatureLevel;
+						AppManifest.EncryptedData.FeatureLevel = AppManifest.ManifestMeta.FeatureLevel;
+					}
+
 					*RawAr << AppManifest.ChunkDataList;
 					*RawAr << AppManifest.FileManifestList;
 					*RawAr << AppManifest.CustomFields;
 					bSuccess = !RawAr->IsError();
+
+					// If we support encryption, serialise the encryption data if necessary.
+					if (bSuccess && Header.Version >= EFeatureLevel::ManifestEncryptionSupport)
+					{
+						if (Header.EncryptionSecretId.IsValid())
+						{
+							*RawAr << AppManifest.EncryptedData;
+							bSuccess = !RawAr->IsError();
+							AppManifest.bIsManifestEncrypted = AppManifest.EncryptedData.Data.Num() > 0;
+						}
+					}
+
 					//// Here we would check for later header versions to serialise additional structures.
 					//if (bSuccess && Header.Version >= EFeatureLevel::SomeShinyNewVersion)
 					//{
 					//	*RawAr << AppManifest.CustomFields;
 					//	bSuccess = !RawAr->IsError();
 					//}
+
+					// If we are saving, restore the original manifest format on each complex serialisable member
+					if (Ar.IsSaving())
+					{
+						AppManifest.ManifestMeta.FeatureLevel = OriginalFormat;
+						AppManifest.ChunkDataList.FeatureLevel = OriginalFormat;
+						AppManifest.FileManifestList.FeatureLevel = OriginalFormat;
+						AppManifest.CustomFields.FeatureLevel = OriginalFormat;
+						AppManifest.EncryptedData.FeatureLevel = OriginalFormat;
+					}
 				}
 				// If saving, calculate the raw data SHA.
 				if (bSuccess && Ar.IsSaving())
@@ -655,30 +1046,21 @@ namespace BuildPatchServices
 					FSHA1::HashBuffer(ManifestRawData.GetData(), ManifestRawData.Num(), DataHash.Hash);
 					Header.SHAHash = DataHash;
 				}
-				// Compress to input archive.
+				// Compress if saving.
 				if (bSuccess && Ar.IsSaving())
 				{
-					TArray<uint8> TempCompressed;
 					Header.DataSizeUncompressed = ManifestRawData.Num();
+					Header.StoredAs = ManifestDataHelpers::CompressManifestMemory(ManifestRawData);
 					Header.DataSizeCompressed = ManifestRawData.Num();
-					TempCompressed.AddUninitialized(Header.DataSizeCompressed);
-					const bool bDataIsCompressed = FCompression::CompressMemory(
-						CompressionFormat,
-						TempCompressed.GetData(),
-						(int32&)Header.DataSizeCompressed,
-						ManifestRawData.GetData(),
-						ManifestRawData.Num(),
-						CompressionFlags);
-					if (bDataIsCompressed)
-					{
-						TempCompressed.SetNum(Header.DataSizeCompressed, EAllowShrinking::No);
-						ManifestRawData = MoveTemp(TempCompressed);
-						Header.StoredAs = EManifestStorageFlags::Compressed;
-					}
-					else
+				}
+				// Set encryption info if saving.
+				if (bSuccess && Ar.IsSaving() && AppManifest.EncryptionSecretId.IsValid())
+				{
+					Header.EncryptionSecretId = AppManifest.EncryptionSecretId;
+					FMemory::Memcpy(Header.EncryptionAuthTag, AppManifest.EncryptionAuthTag.AuthTag, AES256_GCM_AuthTagSizeInBytes);
+					if (AppManifest.bIsManifestEncrypted)
 					{
-						Header.DataSizeCompressed = ManifestRawData.Num();
-						Header.StoredAs = EManifestStorageFlags::None;
+						EnumAddFlags(Header.StoredAs, EManifestStorageFlags::Encrypted);
 					}
 				}
 				// If saving, go back to fill out header and then write data.
diff --git a/Engine/Source/Runtime/Online/BuildPatchServices/Private/Data/ManifestData.h b/Engine/Source/Runtime/Online/BuildPatchServices/Private/Data/ManifestData.h
--- a/Engine/Source/Runtime/Online/BuildPatchServices/Private/Data/ManifestData.h
+++ b/Engine/Source/Runtime/Online/BuildPatchServices/Private/Data/ManifestData.h
@@ -6,6 +6,8 @@
 #include "Misc/EnumClassFlags.h"
 #include "Misc/EnumRange.h"
 #include "Serialization/Archive.h"
+
+#include "Interfaces/IBuildManifest.h"
 #include "Data/ChunkData.h"
 #include "BuildPatchFeatureLevel.h"
 
@@ -14,16 +16,17 @@ class FBuildPatchAppManifest;
 namespace BuildPatchServices
 {
 	/**
-	 * A flags enum for manifest headers which specify storage types.
+	 * Declares flags for manifest headers which specify storage types.
 	 */
 	enum class EManifestStorageFlags : uint8
 	{
-		// Stored as raw data.
-		None       = 0,
+		RawData = 0x00,
+
 		// Flag for compressed data.
-		Compressed = 1,
+		Compressed = 0x01,
+
 		// Flag for encrypted. If also compressed, decrypt first. Encryption will ruin compressibility.
-		Encrypted  = 1 << 1,
+		Encrypted = 0x02
 	};
 	ENUM_CLASS_FLAGS(EManifestStorageFlags);
 
@@ -47,6 +50,16 @@ namespace BuildPatchServices
 		const TCHAR* GetFileSubdir(EFeatureLevel FeatureLevel);
 	}
 
+	namespace ManifestDataHelpers
+	{
+		EManifestStorageFlags CompressManifestMemory(TArray<uint8>& Data);
+		bool UncompressManifestMemory(TArray<uint8>& Data, uint32 DataSizeUncompressed);
+		FString ObfuscateString(const FString& String);
+		TArray<FString> ObfuscateStrings(const TArray<FString>& String);
+		TSet<FString> ObfuscateStrings(const TSet<FString>& Strings);
+	}
+
+
 	struct FManifestHeader
 	{
 		FManifestHeader();
@@ -69,6 +82,10 @@ namespace BuildPatchServices
 		EManifestStorageFlags StoredAs;
 		// The SHA1 hash for the manifest data that follows.
 		FSHAHash SHAHash;
+		// An ID that identifies the encryption secret required to decrypt certain manifest fields.
+		// Without having the associated key, only chunk use, file hashes, and sizes can be understood.
+		FGuid EncryptionSecretId;
+		uint8 EncryptionAuthTag[AES256_GCM_AuthTagSizeInBytes];
 	};
 
 	struct FManifestMeta
@@ -92,17 +109,29 @@ namespace BuildPatchServices
 		// The build version string provided at generation.
 		FString BuildVersion;
 		// The file in this manifest designated the application executable of the build.
+		// Can be an obfuscated string for encrypted manifests.
 		FString LaunchExe;
 		// The command line required when launching the application executable.
+		// Can be an obfuscated string for encrypted manifests.
 		FString LaunchCommand;
 		// The set of prerequisite ids for dependencies that this build's prerequisite installer will apply.
+		// Can be an obfuscated strings for encrypted manifests.
 		TSet<FString> PrereqIds;
 		// A display string for the prerequisite provided at generation.
+		// Can be an obfuscated string for encrypted manifests.
 		FString PrereqName;
 		// The file in this manifest designated the launch executable of the prerequisite installer.
+		// Can be an obfuscated string for encrypted manifests.
 		FString PrereqPath;
 		// The command line required when launching the prerequisite installer.
+		// Can be an obfuscated string for encrypted manifests.
 		FString PrereqArgs;
+		// The path to the uninstall custom action executable.
+		// Can be an obfuscated string for encrypted manifests.
+		FString UninstallActionPath;
+		// The arguments to the uninstall custom action executable.
+		// Can be an obfuscated string for encrypted manifests.
+		FString UninstallActionArgs;
 		// A unique build id generated at original chunking time to identify an exact build.
 		FString BuildId;
 	};
@@ -117,34 +146,29 @@ namespace BuildPatchServices
 		 * @return Passed in archive.
 		 */
 		friend FArchive& operator<<(FArchive& Ar, FChunkDataList& ChunkDataList);
+		// The feature level support this build was created with, regardless of the serialised format.
+		EFeatureLevel FeatureLevel;
 		// The list of chunks.
 		TArray<FChunkInfo> ChunkList;
 	};
 
-	/**
-	 * Declares flags for manifest headers which specify storage types.
-	 */
-	enum class EFileMetaFlags : uint8
-	{
-		None           = 0,
-		// Flag for readonly file.
-		ReadOnly       = 1,
-		// Flag for natively compressed.
-		Compressed     = 1 << 1,
-		// Flag for unix executable.
-		UnixExecutable = 1 << 2
-	};
-	ENUM_CLASS_FLAGS(EFileMetaFlags);
-
 	struct FFileManifest
 	{
 		FFileManifest();
 		// The build relative filename.
+		// Can be an obfuscated string for encrypted manifests.
 		FString Filename;
 		// Whether this is a symlink to another file.
+		// Can be an obfuscated string for encrypted manifests.
 		FString SymlinkTarget;
+		// The file MD5.
+		FMD5Hash MD5Hash;
 		// The file SHA1.
-		FSHAHash FileHash;
+		FSHAHash SHA1Hash;
+		// The file SHA256.
+		FSHA256Signature SHA256Hash;
+		// The calculated MIME type for the file.
+		FString MIMEType;
 		// The flags for this file.
 		EFileMetaFlags FileMetaFlags;
 		// The install tags for this file.
@@ -165,6 +189,8 @@ namespace BuildPatchServices
 		 * @return Passed in archive.
 		 */
 		friend FArchive& operator<<(FArchive& Ar, FFileManifestList& FileManifestList);
+		// The feature level support this build was created with, regardless of the serialised format.
+		EFeatureLevel FeatureLevel;
 		/**
 		 * Helper to sort and calculate file sizes after loading.
 		 */
@@ -183,10 +209,28 @@ namespace BuildPatchServices
 		 * @return Passed in archive.
 		 */
 		friend FArchive& operator<<(FArchive& Ar, FCustomFields& CustomFields);
+		// The feature level support this build was created with, regardless of the serialised format.
+		EFeatureLevel FeatureLevel;
 		// The map of field name to field data.
 		TMap<FString, FString> Fields;
 	};
 
+	struct FEncryptedData
+	{
+		FEncryptedData();
+		/**
+		 * Serialization operator.
+		 * @param Ar             Archive to serialize to.
+		 * @param EncryptedData  FEncryptedData to serialize.
+		 * @return Passed in archive.
+		 */
+		friend FArchive& operator<<(FArchive& Ar, FEncryptedData& EncryptedData);
+		// The feature level support this build was created with, regardless of the serialised format.
+		EFeatureLevel FeatureLevel;
+		// The raw ciphertext.
+		TArray<uint8> Data;
+	};
+
 	class FManifestData
 	{
 	public:
diff --git a/Engine/Source/Runtime/Online/BuildPatchServices/Private/Data/ManifestUObject.cpp b/Engine/Source/Runtime/Online/BuildPatchServices/Private/Data/ManifestUObject.cpp
--- a/Engine/Source/Runtime/Online/BuildPatchServices/Private/Data/ManifestUObject.cpp
+++ b/Engine/Source/Runtime/Online/BuildPatchServices/Private/Data/ManifestUObject.cpp
@@ -65,8 +65,6 @@ FChunkPartData::FChunkPartData()
 *****************************************************************************/
 FFileManifestData::FFileManifestData()
 	: Filename(TEXT(""))
-	, FileHash()
-	, FileChunkParts()
 	, bIsUnixExecutable(false)
 	, SymlinkTarget(TEXT(""))
 	, bIsReadOnly(false)
@@ -101,6 +99,14 @@ UBuildPatchManifest::~UBuildPatchManifest()
 
 namespace ManifestUObjectHelpers
 {
+	UBuildPatchManifest* CreateUBuildPatchManifestFromAsync()
+	{
+		FGCScopeGuard GCGuard;
+		UBuildPatchManifest* Data = NewObject<UBuildPatchManifest>();
+		Data->AddToRoot();
+		return Data;
+	}
+
 	BuildPatchServices::FChunkPart FromChunkPartData(const FChunkPartData& Input)
 	{
 		BuildPatchServices::FChunkPart Output;
@@ -114,14 +120,14 @@ namespace ManifestUObjectHelpers
 	{
 		BuildPatchServices::FFileManifest Output;
 		Output.Filename = Input.Filename;
-		FMemory::Memcpy(Output.FileHash.Hash, Input.FileHash.Hash, FSHA1::DigestSize);
+		FMemory::Memcpy(Output.SHA1Hash.Hash, Input.FileHash.Hash, FSHA1::DigestSize);
 		Output.ChunkParts.Empty(Input.FileChunkParts.Num());
 		Algo::Transform(Input.FileChunkParts, Output.ChunkParts, &ManifestUObjectHelpers::FromChunkPartData);
 		Output.InstallTags = Input.InstallTags;
 		Output.SymlinkTarget = Input.SymlinkTarget;
-		Output.FileMetaFlags |= Input.bIsReadOnly ? BuildPatchServices::EFileMetaFlags::ReadOnly : BuildPatchServices::EFileMetaFlags::None;
-		Output.FileMetaFlags |= Input.bIsCompressed ? BuildPatchServices::EFileMetaFlags::Compressed : BuildPatchServices::EFileMetaFlags::None;
-		Output.FileMetaFlags |= Input.bIsUnixExecutable ? BuildPatchServices::EFileMetaFlags::UnixExecutable : BuildPatchServices::EFileMetaFlags::None;
+		Output.FileMetaFlags |= Input.bIsReadOnly ? EFileMetaFlags::ReadOnly : EFileMetaFlags::None;
+		Output.FileMetaFlags |= Input.bIsCompressed ? EFileMetaFlags::Compressed : EFileMetaFlags::None;
+		Output.FileMetaFlags |= Input.bIsUnixExecutable ? EFileMetaFlags::UnixExecutable : EFileMetaFlags::None;
 		return Output;
 	}
 
@@ -149,14 +155,14 @@ namespace ManifestUObjectHelpers
 	{
 		FFileManifestData Output;
 		Output.Filename = Input.Filename;
-		FMemory::Memcpy(Output.FileHash.Hash, Input.FileHash.Hash, FSHA1::DigestSize);
+		FMemory::Memcpy(Output.FileHash.Hash, Input.SHA1Hash.Hash, FSHA1::DigestSize);
 		Output.FileChunkParts.Empty(Input.ChunkParts.Num());
 		Algo::Transform(Input.ChunkParts, Output.FileChunkParts, &ManifestUObjectHelpers::ToChunkPartData);
 		Output.InstallTags = Input.InstallTags;
-		Output.bIsUnixExecutable = EnumHasAllFlags(Input.FileMetaFlags, BuildPatchServices::EFileMetaFlags::UnixExecutable);
+		Output.bIsUnixExecutable = EnumHasAllFlags(Input.FileMetaFlags, EFileMetaFlags::UnixExecutable);
 		Output.SymlinkTarget = Input.SymlinkTarget;
-		Output.bIsReadOnly = EnumHasAllFlags(Input.FileMetaFlags, BuildPatchServices::EFileMetaFlags::ReadOnly);
-		Output.bIsCompressed = EnumHasAllFlags(Input.FileMetaFlags, BuildPatchServices::EFileMetaFlags::Compressed);
+		Output.bIsReadOnly = EnumHasAllFlags(Input.FileMetaFlags, EFileMetaFlags::ReadOnly);
+		Output.bIsCompressed = EnumHasAllFlags(Input.FileMetaFlags, EFileMetaFlags::Compressed);
 		return Output;
 	}
 
@@ -392,6 +398,7 @@ void FManifestUObject::Init()
 	// We make sure that NewObject<UBuildPatchManifest>() has been called from main thread before it can be called for the 'first time' concurrently
 	// on multiple threads, otherwise a race condition can hit unprotected Emplace on UPackage::ClassUniqueNameIndexMap.
 	// The object will be GC'd on next GC run.
+	check(IsInGameThread());
 	NewObject<UBuildPatchManifest>();
 #endif // !BUILDPATCHSERVICES_NOUOBJECT
 }
@@ -485,7 +492,7 @@ bool FManifestUObject::SaveToArchive(FArchive& Ar, const FBuildPatchAppManifest&
 
 			FManifestHeader Header;
 			Header.Version = AppManifest.ManifestMeta.FeatureLevel;
-			Header.StoredAs = bDataIsCompressed ? EManifestStorageFlags::Compressed : EManifestStorageFlags::None;
+			Header.StoredAs = bDataIsCompressed ? EManifestStorageFlags::Compressed : EManifestStorageFlags::RawData;
 			Header.DataSizeUncompressed = DataSize;
 			Header.DataSizeCompressed = bDataIsCompressed ? CompressedSize : Header.DataSizeUncompressed;
 			FSHA1::HashBuffer(FileData.GetData(), FileData.Num(), Header.SHAHash.Hash);
@@ -517,8 +524,7 @@ bool FManifestUObject::LoadInternal(FArchive& Ar, FBuildPatchAppManifest& AppMan
 	using namespace BuildPatchServices;
 #if !BUILDPATCHSERVICES_NOUOBJECT
 
-	UBuildPatchManifest* Data = NewObject<UBuildPatchManifest>();
-	Data->AddToRoot();
+	UBuildPatchManifest* Data = ManifestUObjectHelpers::CreateUBuildPatchManifestFromAsync();
 
 	// Make sure we use the correct serialization version, this is now fixed and must never use a newer version,
 	// because the property tag has changed in structure meaning older clients would not read correctly.
@@ -600,8 +606,7 @@ bool FManifestUObject::SaveInternal(FArchive& Ar, const FBuildPatchAppManifest&
 {
 #if !BUILDPATCHSERVICES_NOUOBJECT
 
-	UBuildPatchManifest* Data = NewObject<UBuildPatchManifest>();
-	Data->AddToRoot();
+	UBuildPatchManifest* Data = ManifestUObjectHelpers::CreateUBuildPatchManifestFromAsync();
 
 	// Make sure we use the correct serialization version, this is now fixed and must never use a newer version,
 	// because the property tag has changed in structure meaning older clients would not read correctly.
