p4 integrate //UE5/Dev-Distro-5.5/Engine/Source/Programs/BuildPatchTool/... //UE5/Main/Engine/Source/Programs/BuildPatchTool/...

diff --git a/Engine/Source/Runtime/Online/BuildPatchServices/Private/Diffing/DiffManifests.cpp b/Engine/Source/Runtime/Online/BuildPatchServices/Private/Diffing/DiffManifests.cpp
--- a/Engine/Source/Runtime/Online/BuildPatchServices/Private/Diffing/DiffManifests.cpp
+++ b/Engine/Source/Runtime/Online/BuildPatchServices/Private/Diffing/DiffManifests.cpp
@@ -23,11 +23,14 @@
 #include "Installer/DownloadService.h"
 #include "Installer/InstallerAnalytics.h"
 #include "Installer/OptimisedDelta.h"
+#include "Installer/MessagePump.h"
 #include "BuildPatchFileConstructor.h"
 #include "BuildPatchManifest.h"
 #include "BuildPatchUtil.h"
+#include "Tasks/Task.h"
 
-DECLARE_LOG_CATEGORY_CLASS(LogDiffManifests, Log, All);
+DECLARE_LOG_CATEGORY_EXTERN(LogDiffManifests, Log, All);
+DEFINE_LOG_CATEGORY(LogDiffManifests);
 
 // For the output file we'll use pretty json in debug, otherwise condensed.
 #if UE_BUILD_DEBUG
@@ -98,13 +101,12 @@ namespace BuildPatchServices
 				{
 					const FFileManifest* const OldFile = CurrentManifest->GetFileManifest(TaggedFile);
 					const FFileManifest* const NewFile = InstallManifest->GetFileManifest(TaggedFile);
-					if (!OldFile || OldFile->FileHash != NewFile->FileHash || !FilesInstalled.Contains(TaggedFile))
+					if (!OldFile || OldFile->SHA1Hash != NewFile->SHA1Hash || !FilesInstalled.Contains(TaggedFile))
 					{
 						FilesToBuild.Add(MoveTemp(TaggedFile));
 					}
 				}
 			}
-			FilesToBuild.Sort(TLess<FString>());
 			CurrentManifest->GetChunksRequiredForFiles(FilesToBuild, ChunksNeeded);
 
 			// Setup a chunk reference tracker.
@@ -173,7 +175,7 @@ namespace BuildPatchServices
 						}
 						else if (ChunksInstalled.Contains(ChunkInfo.Guid))
 						{
-							Timer += (double)ChunkInfo.WindowSize / Config.DiskReadSpeed;
+							Timer += (double)ChunkInfo.DataSizeUncompressed / Config.DiskReadSpeed;
 							LoadedChunks.Add(ChunkInfo.Guid);
 						}
 						else
@@ -215,7 +217,7 @@ namespace BuildPatchServices
 						for (const FGuid& BatchLoadChunk : BatchLoadChunks)
 						{
 							// Load it from disk.
-							Timer += (double)InstallManifest.GetChunkInfo(BatchLoadChunk)->WindowSize / Config.BackupSerialisationSpeed;
+							Timer += (double)InstallManifest.GetChunkInfo(BatchLoadChunk)->DataSizeUncompressed / Config.BackupSerialisationSpeed;
 							// Save it to backup.
 							Timer += (double)InstallManifest.GetChunkInfo(BatchLoadChunk)->FileSize / Config.BackupSerialisationSpeed;
 							BackupChunks.Add(BatchLoadChunk);
@@ -300,15 +302,14 @@ namespace BuildPatchServices
 		void ExportPatchDescriptorFiles(const TSet<FString>& InstallTags, FBuildPatchAppManifestPtr OldManifest, FBuildPatchAppManifestPtr NewManifest);
 	private:
 		const FDiffManifestsConfiguration Configuration;
-		FTSTicker& CoreTicker;
 		FDownloadCompleteDelegate DownloadCompleteDelegate;
 		FDownloadProgressDelegate DownloadProgressDelegate;
-		TUniquePtr<IFileSystem> FileSystem;
-		TUniquePtr<IHttpManager> HttpManager;
+		TSharedPtr<IFileSystem> FileSystem;
+		TSharedPtr<IHttpManager> HttpManager;
 		TUniquePtr<IChunkDataSizeProvider> ChunkDataSizeProvider;
 		TUniquePtr<ISpeedRecorder> DownloadSpeedRecorder;
-		TUniquePtr<IInstallerAnalytics> InstallerAnalytics;
-		TUniquePtr<IDownloadServiceStatistics> DownloadServiceStatistics;
+		TSharedPtr<IInstallerAnalytics> InstallerAnalytics;
+		TSharedPtr<IDownloadServiceStatistics> DownloadServiceStatistics;
 		TUniquePtr<IDownloadService> DownloadService;
 		TUniquePtr<FStatsCollector> StatsCollector;
 		FThreadSafeBool bShouldRun;
@@ -320,43 +321,65 @@ namespace BuildPatchServices
 		TPromise<FBuildPatchAppManifestPtr> PromiseManifestB;
 		TFuture<FBuildPatchAppManifestPtr> FutureManifestA;
 		TFuture<FBuildPatchAppManifestPtr> FutureManifestB;
+
+		TUniquePtr<IMessagePump> MessagePump;
+		// Setters will be added in EDS implementation for BPT
+		TArray<FMessageHandler*> MessageHandlers;
 	};
 
 	FDiffManifests::FDiffManifests(const FDiffManifestsConfiguration& InConfiguration)
 		: Configuration(InConfiguration)
-		, CoreTicker(FTSTicker::GetCoreTicker())
 		, DownloadCompleteDelegate(FDownloadCompleteDelegate::CreateRaw(this, &FDiffManifests::HandleDownloadComplete))
 		, DownloadProgressDelegate()
 		, FileSystem(FFileSystemFactory::Create())
 		, HttpManager(FHttpManagerFactory::Create())
-		, ChunkDataSizeProvider(FChunkDataSizeProviderFactory::Create())
+		, ChunkDataSizeProvider(FChunkDataSizeProviderFactory::Create({}))
 		, DownloadSpeedRecorder(FSpeedRecorderFactory::Create())
 		, InstallerAnalytics(FInstallerAnalyticsFactory::Create(nullptr))
 		, DownloadServiceStatistics(FDownloadServiceStatisticsFactory::Create(DownloadSpeedRecorder.Get(), ChunkDataSizeProvider.Get(), InstallerAnalytics.Get()))
-		, DownloadService(FDownloadServiceFactory::Create(HttpManager.Get(), FileSystem.Get(), DownloadServiceStatistics.Get(), InstallerAnalytics.Get()))
+		, DownloadService(FDownloadServiceFactory::Create(HttpManager, FileSystem, DownloadServiceStatistics, InstallerAnalytics))
 		, StatsCollector(FStatsCollectorFactory::Create())
 		, bShouldRun(true)
 		, RequestIdManifestA(INDEX_NONE)
 		, RequestIdManifestB(INDEX_NONE)
-		, PromiseManifestA()
-		, PromiseManifestB()
 		, FutureManifestA(PromiseManifestA.GetFuture())
 		, FutureManifestB(PromiseManifestB.GetFuture())
+		, MessagePump(FMessagePumpFactory::Create())
+		, MessageHandlers(Configuration.MessageHandlers)
 	{
 	}
 
 	FDiffManifests::~FDiffManifests()
 	{
+		if (!FutureManifestA.IsReady())
+		{
+			PromiseManifestA.SetValue(FBuildPatchAppManifestPtr());
+			FutureManifestA.Wait();
+		}
+		if (!FutureManifestB.IsReady())
+		{
+			PromiseManifestB.SetValue(FBuildPatchAppManifestPtr());
+			FutureManifestB.Wait();
+		}
 	}
 
 	bool FDiffManifests::Run()
 	{
+		// Any code that should be ran after the loop exit.
+		ON_SCOPE_EXIT
+		{
+			GLog->FlushThreadedLogs();
+		};
+
 		// Run any core initialisation required.
 		FHttpModule::Get();
 
 		// Kick off Manifest downloads.
-		RequestIdManifestA = DownloadService->RequestFile(Configuration.ManifestAUri, DownloadCompleteDelegate, DownloadProgressDelegate);
-		RequestIdManifestB = DownloadService->RequestFile(Configuration.ManifestBUri, DownloadCompleteDelegate, DownloadProgressDelegate);
+		if(    !DownloadService->RequestFile(Configuration.ManifestAUri, DownloadCompleteDelegate, DownloadProgressDelegate, &RequestIdManifestA)
+		    || !DownloadService->RequestFile(Configuration.ManifestBUri, DownloadCompleteDelegate, DownloadProgressDelegate, &RequestIdManifestB))
+		{
+			return false;
+		}
 
 		// Start the generation thread.
 		TFuture<bool> Thread = Async(EAsyncExecution::Thread, [this]() { return AsyncRun(); });
@@ -378,6 +401,10 @@ namespace BuildPatchServices
 			// Application tick.
 			FTaskGraphInterface::Get().ProcessThreadUntilIdle(ENamedThreads::GameThread);
 			FTSTicker::GetCoreTicker().Tick(DeltaTime);
+
+			// Message pump.
+			MessagePump->PumpMessages(MessageHandlers);
+
 			GLog->FlushThreadedLogs();
 
 			// Control frame rate.
@@ -388,9 +415,8 @@ namespace BuildPatchServices
 			DeltaTime = AppTime - LastTime;
 			LastTime = AppTime;
 		}
-		GLog->FlushThreadedLogs();
 
-		// Return thread success.
+		// Return thread success state.
 		return Thread.Get();
 	}
 
@@ -459,7 +485,7 @@ namespace BuildPatchServices
 			for (FGuid& Guid : ChunksToDownload)
 			{
 				const BuildPatchServices::FChunkInfo* ChunkInfo = NewManifest->GetChunkInfo(Guid);
-				TotalNewChunkBytes += ChunkInfo->WindowSize;
+				TotalNewChunkBytes += ChunkInfo->DataSizeUncompressed;
 				TotalNewChunkDownloadBytes += ChunkInfo->FileSize;
 			}
 
@@ -566,35 +592,42 @@ namespace BuildPatchServices
 	bool FDiffManifests::AsyncRun()
 	{
 		FBuildPatchAppManifestPtr ManifestA = FutureManifestA.Get();
-		FBuildPatchAppManifestPtr ManifestB = FutureManifestB.Get();
+		FBuildPatchAppManifestPtr OriginalManifestB = FutureManifestB.Get();
 		bool bSuccess = true;
 		if (ManifestA.IsValid() == false)
 		{
 			UE_LOG(LogDiffManifests, Error, TEXT("Could not download ManifestA from %s."), *Configuration.ManifestAUri);
 			bSuccess = false;
 		}
-		if (ManifestB.IsValid() == false)
+		if (OriginalManifestB.IsValid() == false)
 		{
 			UE_LOG(LogDiffManifests, Error, TEXT("Could not download ManifestB from %s."), *Configuration.ManifestBUri);
 			bSuccess = false;
 		}
 		if (bSuccess)
 		{
+			TArray<FString> CloudDirectories = Configuration.CloudDirs.Num() == 0 ? TArray<FString>{ FPaths::GetPath(Configuration.ManifestBUri) } : Configuration.CloudDirs;
+
 			// Check for delta file, replacing ManifestB if we find one
-			FOptimisedDeltaConfiguration OptimisedDeltaConfiguration(ManifestB.ToSharedRef());
+			FOptimisedDeltaConfiguration OptimisedDeltaConfiguration(OriginalManifestB.ToSharedRef());
 			OptimisedDeltaConfiguration.SourceManifest = ManifestA;
 			OptimisedDeltaConfiguration.DeltaPolicy = Configuration.bRequireOptimizedDelta ? EDeltaPolicy::Expect : EDeltaPolicy::TryFetchContinueWithout;
-			OptimisedDeltaConfiguration.CloudDirectories = { FPaths::GetPath(Configuration.ManifestBUri) };
-			FOptimisedDeltaDependencies OptimisedDeltaDependencies;
+			OptimisedDeltaConfiguration.RetriesNumber = CloudDirectories.Num();
+			OptimisedDeltaConfiguration.DeltaFilenameTrailer = Configuration.DeltaFilenameTrailer;
+
+			TSharedRef<IUriProvider, ESPMode::ThreadSafe> UriProvider = FUriProviderFactory::Create(MessagePump.Get(), CloudDirectories);
+			FOptimisedDeltaDependencies OptimisedDeltaDependencies(MoveTemp(UriProvider));
 			OptimisedDeltaDependencies.DownloadService = DownloadService.Get();
-			TUniquePtr<IOptimisedDelta> OptimisedDelta(FOptimisedDeltaFactory::Create(OptimisedDeltaConfiguration, MoveTemp(OptimisedDeltaDependencies)));
+			OptimisedDeltaDependencies.UriProvider = MoveTemp(UriProvider);
+
+			TSharedRef<IOptimisedDelta> OptimisedDelta(FOptimisedDeltaFactory::Create(OptimisedDeltaConfiguration, MoveTemp(OptimisedDeltaDependencies)));
 			if (OptimisedDelta->GetResult().HasError())
 			{
 				UE_LOG(LogDiffManifests, Error, TEXT("Optimized delta not found and RequireOptimizedDelta was passed - failing"));
 				bShouldRun = false;
 				return false;
 			}
-			ManifestB = OptimisedDelta->GetResult().GetValue();
+			FBuildPatchAppManifestPtr OptimisedDeltaManifestB = OptimisedDelta->GetResult().GetValue();
 			const int32 MetaDownloadBytes = OptimisedDelta->GetMetaDownloadSize();
 
 			TSet<FString> TagsA, TagsB;
@@ -603,7 +636,7 @@ namespace BuildPatchServices
 			{
 				TagsA = TagsA.Intersect(Configuration.TagSetA);
 			}
-			ManifestB->GetFileTagList(TagsB);
+			OptimisedDeltaManifestB->GetFileTagList(TagsB);
 			if (Configuration.TagSetB.Num() > 0)
 			{
 				TagsB = TagsB.Intersect(Configuration.TagSetB);
@@ -612,7 +645,7 @@ namespace BuildPatchServices
 			if (Configuration.OutputPatchDescriptorPath.Len())
 			{
 				// We want to use all NEW tags so that we catch new bundles, not just changed bundles.
-				ExportPatchDescriptorFiles(TagsB, ManifestA, ManifestB);
+				ExportPatchDescriptorFiles(TagsB, ManifestA, OptimisedDeltaManifestB);
 
 				if (Configuration.bOnlyPatchDescriptors)
 				{
@@ -630,23 +663,23 @@ namespace BuildPatchServices
 			TSet<FGuid> ChunkSetB;
 			ManifestA->GetTaggedFileList(TagsA, TaggedFileSetA);
 			ManifestA->GetChunksRequiredForFiles(TaggedFileSetA, ChunkSetA);
-			ManifestB->GetTaggedFileList(TagsB, TaggedFileSetB);
-			ManifestB->GetChunksRequiredForFiles(TaggedFileSetB, ChunkSetB);
+			OptimisedDeltaManifestB->GetTaggedFileList(TagsB, TaggedFileSetB);
+			OptimisedDeltaManifestB->GetChunksRequiredForFiles(TaggedFileSetB, ChunkSetB);
 			TArray<FString> NewChunkPaths;
 			for (FGuid& ChunkB : ChunkSetB)
 			{
 				if (ChunkSetA.Contains(ChunkB) == false)
 				{
 					++NewChunksCount;
-					int32 ChunkFileSize = ManifestB->GetDataSize(ChunkB);
+					int32 ChunkFileSize = OptimisedDeltaManifestB->GetDataSize(ChunkB);
 					TotalChunkSize += ChunkFileSize;
-					NewChunkPaths.Add(FBuildPatchUtils::GetDataFilename(ManifestB.ToSharedRef(), ChunkB));
+					NewChunkPaths.Add(FBuildPatchUtils::GetDataFilename(OptimisedDeltaManifestB.ToSharedRef(), ChunkB));
 					UE_LOG(LogDiffManifests, Verbose, TEXT("New chunk discovered: Size: %10lld, Path: %s"), ChunkFileSize, *NewChunkPaths.Last());
 				}
 			}
 
-			UE_LOG(LogDiffManifests, Display, TEXT("New chunks:  %lld"), NewChunksCount);
-			UE_LOG(LogDiffManifests, Display, TEXT("Total bytes: %lld"), TotalChunkSize);
+			UE_LOG(LogDiffManifests, Log, TEXT("New chunks:  %lld"), NewChunksCount);
+			UE_LOG(LogDiffManifests, Log, TEXT("Total bytes: %lld"), TotalChunkSize);
 
 			TSet<FString> NewFilePaths = TaggedFileSetB.Difference(TaggedFileSetA);
 			TSet<FString> RemovedFilePaths = TaggedFileSetA.Difference(TaggedFileSetB);
@@ -656,32 +689,34 @@ namespace BuildPatchServices
 			const TSet<FString>& SetToIterate = TaggedFileSetB.Num() > TaggedFileSetA.Num() ? TaggedFileSetA : TaggedFileSetB;
 			for (const FString& TaggedFile : SetToIterate)
 			{
-				FSHAHash FileHashA;
-				FSHAHash FileHashB;
-				if (ManifestA->GetFileHash(TaggedFile, FileHashA) && ManifestB->GetFileHash(TaggedFile, FileHashB))
+				if (!RemovedFilePaths.Contains(TaggedFile) && !NewFilePaths.Contains(TaggedFile))
 				{
-					if (FileHashA == FileHashB)
-					{
-						UnchangedFilePaths.Add(TaggedFile);
-					}
-					else
+					FSHAHash FileHashA;
+					FSHAHash FileHashB;
+					if (ManifestA->GetFileHash(TaggedFile, FileHashA) && OptimisedDeltaManifestB->GetFileHash(TaggedFile, FileHashB))
 					{
-						ChangedFilePaths.Add(TaggedFile);
+						if (FileHashA == FileHashB)
+						{
+							UnchangedFilePaths.Add(TaggedFile);
+						}
+						else
+						{
+							ChangedFilePaths.Add(TaggedFile);
+						}
 					}
 				}
 			}
 
 			// Log download details.
-			FNumberFormattingOptions SizeFormattingOptions;
-			SizeFormattingOptions.MaximumFractionalDigits = 3;
-			SizeFormattingOptions.MinimumFractionalDigits = 3;
+			const FNumberFormattingOptions SizeFormattingOptions = FNumberFormattingOptions().SetMaximumFractionalDigits(3).SetMinimumFractionalDigits(3);
+			const FNumberFormattingOptions PercentFormat = FNumberFormattingOptions().SetMaximumFractionalDigits(1).SetMinimumFractionalDigits(1).SetRoundingMode(ERoundingMode::ToZero);
 
-			int64 DownloadSizeA = ManifestA->GetDownloadSize(TagsA);
-			int64 BuildSizeA = ManifestA->GetBuildSize(TagsA);
-			int64 DownloadSizeB = ManifestB->GetDownloadSize(TagsB);
-			int64 BuildSizeB = ManifestB->GetBuildSize(TagsB);
-			int64 DeltaDownloadSize = ManifestB->GetDeltaDownloadSize(TagsB, ManifestA.ToSharedRef(), TagsA) + MetaDownloadBytes;
-			int64 TempDiskSpaceReq = FileConstructorHelpers::CalculateRequiredDiskSpace(ManifestA, ManifestB.ToSharedRef(), EInstallMode::DestructiveInstall, TagsB);
+			const int64 DownloadSizeA = ManifestA->GetDownloadSize(TagsA);
+			const int64 BuildSizeA = ManifestA->GetBuildSize(TagsA);
+			const int64 DownloadSizeB = OriginalManifestB->GetDownloadSize(TagsB);
+			const int64 BuildSizeB = OriginalManifestB->GetBuildSize(TagsB);
+			const int64 DeltaDownloadSize = OptimisedDeltaManifestB->GetDeltaDownloadSize(TagsB, ManifestA.ToSharedRef(), TagsA) + MetaDownloadBytes;
+			const int64 TempDiskSpaceReq = FileConstructorHelpers::CalculateRequiredDiskSpace(ManifestA, OptimisedDeltaManifestB.ToSharedRef(), EInstallMode::DestructiveInstall, TagsB);
 
 			// Break down the sizes and delta into new chunks per tag.
 			TMap<FString, int64> TagDownloadImpactA;
@@ -691,18 +726,25 @@ namespace BuildPatchServices
 			TMap<FString, int64> TagDeltaImpact;
 			for (const FString& Tag : TagsA)
 			{
-				TSet<FString> TagSet;
-				TagSet.Add(Tag);
-				TagDownloadImpactA.Add(Tag, ManifestA->GetDownloadSize(TagSet));
-				TagBuildImpactA.Add(Tag, ManifestA->GetBuildSize(TagSet));
+				TSet<const FFileManifest*> TaggedManifests = ManifestA->GetTaggedFileManifests(TSet<FString>{ Tag });
+				TagDownloadImpactA.Add(Tag, ManifestA->GetDownloadSize(TaggedManifests));
+				TagBuildImpactA.Add(Tag, ManifestA->GetBuildSize(TaggedManifests));
 			}
+
+			TSet<const FFileManifest*> TaggedManifestsA = ManifestA->GetTaggedFileManifests(TagsA);
+			TSet<FString> FilesInstalled;
+			Algo::Transform(TaggedManifestsA, FilesInstalled, [](const FFileManifest* File) { return File->Filename; });
+			TSet<FGuid> ChunksRequired;
+			TSet<FGuid> ChunksInstalled;
+			ManifestA->GetChunksRequiredForFiles(FilesInstalled, ChunksRequired);
+			ManifestA->EnumerateProducibleChunks(TaggedManifestsA, ChunksRequired, ChunksInstalled);
+
 			for (const FString& Tag : TagsB)
 			{
-				TSet<FString> TagSet;
-				TagSet.Add(Tag);
-				TagDownloadImpactB.Add(Tag, ManifestB->GetDownloadSize(TagSet));
-				TagBuildImpactB.Add(Tag, ManifestB->GetBuildSize(TagSet));
-				TagDeltaImpact.Add(Tag, ManifestB->GetDeltaDownloadSize(TagSet, ManifestA.ToSharedRef(), TagsA));
+				TSet<const FFileManifest*> TaggedManifests = OriginalManifestB->GetTaggedFileManifests(TSet<FString>{ Tag });
+				TagDownloadImpactB.Add(Tag, OriginalManifestB->GetDownloadSize(TaggedManifests));
+				TagBuildImpactB.Add(Tag, OriginalManifestB->GetBuildSize(TaggedManifests));
+				TagDeltaImpact.Add(Tag, OptimisedDeltaManifestB->GetDeltaDownloadSize(TaggedManifests, ManifestA.ToSharedRef(), FilesInstalled, ChunksInstalled));
 			}
 			if (MetaDownloadBytes > 0)
 			{
@@ -717,42 +759,54 @@ namespace BuildPatchServices
 			TMap<FString, int64> CompareTagSetDownloadSizeB;
 			TMap<FString, int64> CompareTagSetTempDiskSpaceReqs;
 			TSet<FString> CompareTagSetKeys;
-			for (const TSet<FString>& TagSet : Configuration.CompareTagSets)
+
+			for (int32 TagSetIndex = 0; TagSetIndex < Configuration.CompareTagSetsA.Num(); TagSetIndex++)
 			{
-				TArray<FString> TagArrayCompare = TagSet.Array();
+				const TSet<FString>& OldTagSet = Configuration.CompareTagSetsA[TagSetIndex];
+
+				// If the set of installed tags changes across the patch, we need to pass both sets to GetDeltaDownloadSize.
+				const TSet<FString>* NewTagSet = &OldTagSet;
+				if (TagSetIndex < Configuration.CompareTagSetsB.Num())
+				{
+					NewTagSet = &Configuration.CompareTagSetsB[TagSetIndex];
+				}
+
+				TArray<FString> TagArrayCompare = NewTagSet->Array();
 				Algo::Sort(TagArrayCompare);
 				FString TagSetString = FString::Join(TagArrayCompare, TEXT(", "));
+				if (&OldTagSet != NewTagSet)
+				{
+					TArray<FString> OldTagArray = OldTagSet.Array();
+					Algo::Sort(OldTagArray);
+					FString OldTagSetString = FString::Join(OldTagArray, TEXT(", "));
+					TagSetString = OldTagSetString + TEXT(" -> ") + TagSetString;
+				}
+
 				CompareTagSetKeys.Add(TagSetString);
-				CompareTagSetDeltaImpact.Add(TagSetString, ManifestB->GetDeltaDownloadSize(TagSet, ManifestA.ToSharedRef(), TagSet) + MetaDownloadBytes);
-				CompareTagSetBuildImpactB.Add(TagSetString, ManifestB->GetBuildSize(TagSet));
-				CompareTagSetDownloadSizeB.Add(TagSetString, ManifestB->GetDownloadSize(TagSet));
-				CompareTagSetBuildImpactA.Add(TagSetString, ManifestA->GetBuildSize(TagSet));
-				CompareTagSetDownloadSizeA.Add(TagSetString, ManifestA->GetDownloadSize(TagSet));
-				CompareTagSetTempDiskSpaceReqs.Add(TagSetString, FileConstructorHelpers::CalculateRequiredDiskSpace(ManifestA, ManifestB.ToSharedRef(), EInstallMode::DestructiveInstall, TagSet));
+
+				CompareTagSetDeltaImpact.Add(TagSetString, OptimisedDeltaManifestB->GetDeltaDownloadSize(*NewTagSet, ManifestA.ToSharedRef(), OldTagSet) + MetaDownloadBytes);
+				CompareTagSetBuildImpactB.Add(TagSetString, OriginalManifestB->GetBuildSize(*NewTagSet));
+				CompareTagSetDownloadSizeB.Add(TagSetString, OriginalManifestB->GetDownloadSize(*NewTagSet));
+
+				CompareTagSetBuildImpactA.Add(TagSetString, ManifestA->GetBuildSize(OldTagSet));
+				CompareTagSetDownloadSizeA.Add(TagSetString, ManifestA->GetDownloadSize(OldTagSet));
+				CompareTagSetTempDiskSpaceReqs.Add(TagSetString, FileConstructorHelpers::CalculateRequiredDiskSpace(ManifestA, OptimisedDeltaManifestB.ToSharedRef(), EInstallMode::DestructiveInstall, *NewTagSet));
 			}
 
 			// Log the information.
-			TArray<FString> TagArrayB = TagsB.Array();
-			Algo::Sort(TagArrayB);
-			FString UntaggedLog("(untagged)");
-			FString TagLogList = FString::Join(TagArrayB, TEXT(", "));
-			if (TagLogList.IsEmpty() || TagLogList.StartsWith(TEXT(", ")))
-			{
-				TagLogList.InsertAt(0, UntaggedLog);
-			}
-			UE_LOG(LogDiffManifests, Display, TEXT("TagSet: %s"), *TagLogList);
 			UE_LOG(LogDiffManifests, Display, TEXT("%s %s:"), *ManifestA->GetAppName(), *ManifestA->GetVersionString());
 			UE_LOG(LogDiffManifests, Display, TEXT("    Download Size:   %20s bytes (%10s, %11s)"), *FText::AsNumber(DownloadSizeA).ToString(), *FText::AsMemory(DownloadSizeA, &SizeFormattingOptions, nullptr, EMemoryUnitStandard::SI).ToString(), *FText::AsMemory(DownloadSizeA, &SizeFormattingOptions, nullptr, EMemoryUnitStandard::IEC).ToString());
 			UE_LOG(LogDiffManifests, Display, TEXT("    Build Size:      %20s bytes (%10s, %11s)"), *FText::AsNumber(BuildSizeA).ToString(), *FText::AsMemory(BuildSizeA, &SizeFormattingOptions, nullptr, EMemoryUnitStandard::SI).ToString(), *FText::AsMemory(BuildSizeA, &SizeFormattingOptions, nullptr, EMemoryUnitStandard::IEC).ToString());
-			UE_LOG(LogDiffManifests, Display, TEXT("%s %s:"), *ManifestB->GetAppName(), *ManifestB->GetVersionString());
+			UE_LOG(LogDiffManifests, Display, TEXT("%s %s:"), *OriginalManifestB->GetAppName(), *OriginalManifestB->GetVersionString());
 			UE_LOG(LogDiffManifests, Display, TEXT("    Download Size:   %20s bytes (%10s, %11s)"), *FText::AsNumber(DownloadSizeB).ToString(), *FText::AsMemory(DownloadSizeB, &SizeFormattingOptions, nullptr, EMemoryUnitStandard::SI).ToString(), *FText::AsMemory(DownloadSizeB, &SizeFormattingOptions, nullptr, EMemoryUnitStandard::IEC).ToString());
 			UE_LOG(LogDiffManifests, Display, TEXT("    Build Size:      %20s bytes (%10s, %11s)"), *FText::AsNumber(BuildSizeB).ToString(), *FText::AsMemory(BuildSizeB, &SizeFormattingOptions, nullptr, EMemoryUnitStandard::SI).ToString(), *FText::AsMemory(BuildSizeB, &SizeFormattingOptions, nullptr, EMemoryUnitStandard::IEC).ToString());
-			UE_LOG(LogDiffManifests, Display, TEXT("%s %s -> %s %s:"), *ManifestA->GetAppName(), *ManifestA->GetVersionString(), *ManifestB->GetAppName(), *ManifestB->GetVersionString());
+			UE_LOG(LogDiffManifests, Display, TEXT("%s %s -> %s %s:"), *ManifestA->GetAppName(), *ManifestA->GetVersionString(), *OptimisedDeltaManifestB->GetAppName(), *OptimisedDeltaManifestB->GetVersionString());
 			UE_LOG(LogDiffManifests, Display, TEXT("    Delta Size:      %20s bytes (%10s, %11s)"), *FText::AsNumber(DeltaDownloadSize).ToString(), *FText::AsMemory(DeltaDownloadSize, &SizeFormattingOptions, nullptr, EMemoryUnitStandard::SI).ToString(), *FText::AsMemory(DeltaDownloadSize, &SizeFormattingOptions, nullptr, EMemoryUnitStandard::IEC).ToString());
 			UE_LOG(LogDiffManifests, Display, TEXT("    Temp Disk Space: %20s bytes (%10s, %11s)"), *FText::AsNumber(TempDiskSpaceReq).ToString(), *FText::AsMemory(TempDiskSpaceReq, &SizeFormattingOptions, nullptr, EMemoryUnitStandard::SI).ToString(), *FText::AsMemory(TempDiskSpaceReq, &SizeFormattingOptions, nullptr, EMemoryUnitStandard::IEC).ToString());
 			UE_LOG(LogDiffManifests, Display, TEXT(""));
 
-			for (const FString& Tag : TagArrayB)
+			const FString UntaggedLog(TEXT("(untagged)"));
+			for (const FString& Tag : TagsB)
 			{
 				UE_LOG(LogDiffManifests, Display, TEXT("%s Impact:"), *(Tag.IsEmpty() ? UntaggedLog : Tag));
 				UE_LOG(LogDiffManifests, Display, TEXT("    Individual Download Size:  %20s bytes (%10s, %11s)"), *FText::AsNumber(TagDownloadImpactB[Tag]).ToString(), *FText::AsMemory(TagDownloadImpactB[Tag], &SizeFormattingOptions, nullptr, EMemoryUnitStandard::SI).ToString(), *FText::AsMemory(TagDownloadImpactB[Tag], &SizeFormattingOptions, nullptr, EMemoryUnitStandard::IEC).ToString());
@@ -770,10 +824,10 @@ namespace BuildPatchServices
 				UE_LOG(LogDiffManifests, Display, TEXT("    Temp Disk Space:  %20s bytes (%10s, %11s)"), *FText::AsNumber(CompareTagSetTempDiskSpaceReqs[TagSet]).ToString(), *FText::AsMemory(CompareTagSetTempDiskSpaceReqs[TagSet], &SizeFormattingOptions, nullptr, EMemoryUnitStandard::SI).ToString(), *FText::AsMemory(CompareTagSetTempDiskSpaceReqs[TagSet], &SizeFormattingOptions, nullptr, EMemoryUnitStandard::IEC).ToString());
 			}
 
-			// Hit a destructive and nondestructive simulation for a few different specs.
 			TArray<double> InstallTimeCoefficients;
-			if (Configuration.bEmitInstallTime)
+			if (Configuration.bSimulateInstallTimes)
 			{
+				// Hit a destructive and nondestructive simulation for a few different specs.
 				TArray<DiffHelpers::FSimConfig> SimConfigs;
 				// Add some lower spec values, taken from around 25 percentile of stats at the time of writing [July 2019].
 				SimConfigs.AddDefaulted_GetRef().InstallMode = EInstallMode::DestructiveInstall;
@@ -801,7 +855,7 @@ namespace BuildPatchServices
 				SimConfigs[4].BackupSerialisationSpeed = SimConfigs[5].BackupSerialisationSpeed = 40000000.0; // 40 MB/s
 
 				// Run the calculations and log.
-				InstallTimeCoefficients = DiffHelpers::CalculateInstallTimeCoefficient(ManifestA.ToSharedRef(), TagsA, ManifestB.ToSharedRef(), TagsB, SimConfigs);
+				InstallTimeCoefficients = DiffHelpers::CalculateInstallTimeCoefficient(ManifestA.ToSharedRef(), TagsA, OptimisedDeltaManifestB.ToSharedRef(), TagsB, SimConfigs);
 				checkf(6 == InstallTimeCoefficients.Num() && 6 == SimConfigs.Num(), TEXT("Unexpected result size from CalculateInstallTimeCoefficient."));
 				UE_LOG(LogDiffManifests, Display, TEXT(""));
 				UE_LOG(LogDiffManifests, Display, TEXT("Install time coefficients are not accurate timing representations, but are comparable from patch to patch."));
@@ -815,6 +869,28 @@ namespace BuildPatchServices
 				UE_LOG(LogDiffManifests, Display, TEXT("    High-Spec NonDestructiveInstall: %s"), *FPlatformTime::PrettyTime(InstallTimeCoefficients[5]));
 			}
 
+			// Output if an optimization was ran, what the overall improvement was.
+			UE_LOG(LogDiffManifests, Display, TEXT(""));
+			if (MetaDownloadBytes > 0)
+			{
+				const int64 OriginalDeltaDownloadSize = OriginalManifestB->GetDeltaDownloadSize({}, ManifestA.ToSharedRef());
+				const int64 OptimisedDeltaDownloadSize = OptimisedDeltaManifestB->GetDeltaDownloadSize({}, ManifestA.ToSharedRef()) + MetaDownloadBytes;
+				if (OptimisedDeltaDownloadSize < OriginalDeltaDownloadSize)
+				{
+					UE_LOG(LogDiffManifests, Display, TEXT("Delta optimisation achieved a download size improvement of %s."), *FText::AsPercent(1.0 - ((double)OptimisedDeltaDownloadSize / (double)OriginalDeltaDownloadSize), &PercentFormat).ToString());
+					UE_LOG(LogDiffManifests, Display, TEXT("    Original:  %20s bytes (%10s, %11s)."), *FText::AsNumber(OriginalDeltaDownloadSize).ToString(), *FText::AsMemory(OriginalDeltaDownloadSize, &SizeFormattingOptions, nullptr, EMemoryUnitStandard::SI).ToString(), *FText::AsMemory(OriginalDeltaDownloadSize, &SizeFormattingOptions, nullptr, EMemoryUnitStandard::IEC).ToString());
+					UE_LOG(LogDiffManifests, Display, TEXT("    Optimised: %20s bytes (%10s, %11s)."), *FText::AsNumber(OptimisedDeltaDownloadSize).ToString(), *FText::AsMemory(OptimisedDeltaDownloadSize, &SizeFormattingOptions, nullptr, EMemoryUnitStandard::SI).ToString(), *FText::AsMemory(OptimisedDeltaDownloadSize, &SizeFormattingOptions, nullptr, EMemoryUnitStandard::IEC).ToString());
+				}
+				else
+				{
+					UE_LOG(LogDiffManifests, Display, TEXT("Delta optimisation was completed but did not find an improvement."));
+				}
+			}
+			else
+			{
+				UE_LOG(LogDiffManifests, Display, TEXT("Delta optimisation has not been completed for patching %s %s to %s %s."), *ManifestA->GetAppName(), *ManifestA->GetVersionString(), *OriginalManifestB->GetAppName(), *OriginalManifestB->GetVersionString());
+			}
+
 			// Save the output.
 			if (bSuccess && Configuration.OutputFilePath.IsEmpty() == false)
 			{
@@ -857,9 +933,9 @@ namespace BuildPatchServices
 					Writer->WriteObjectEnd();
 					Writer->WriteObjectStart(TEXT("ManifestB"));
 					{
-						Writer->WriteValue(TEXT("AppName"), ManifestB->GetAppName());
-						Writer->WriteValue(TEXT("AppId"), static_cast<int32>(ManifestB->GetAppID()));
-						Writer->WriteValue(TEXT("VersionString"), ManifestB->GetVersionString());
+						Writer->WriteValue(TEXT("AppName"), OriginalManifestB->GetAppName());
+						Writer->WriteValue(TEXT("AppId"), static_cast<int32>(OriginalManifestB->GetAppID()));
+						Writer->WriteValue(TEXT("VersionString"), OriginalManifestB->GetVersionString());
 						Writer->WriteValue(TEXT("DownloadSize"), DownloadSizeB);
 						Writer->WriteValue(TEXT("BuildSize"), BuildSizeB);
 						Writer->WriteObjectStart(TEXT("IndividualTagDownloadSizes"));
@@ -941,12 +1017,15 @@ namespace BuildPatchServices
 							Writer->WriteValue(Pair.Key, Pair.Value);
 						}
 						Writer->WriteObjectEnd();
-						Writer->WriteArrayStart(TEXT("InstallTimeCoefficients"));
-						for (const double& InstallTimeCoefficient : InstallTimeCoefficients)
+						if (Configuration.bSimulateInstallTimes)
 						{
-							Writer->WriteValue(InstallTimeCoefficient);
+							Writer->WriteArrayStart(TEXT("InstallTimeCoefficients"));
+							for (const double& InstallTimeCoefficient : InstallTimeCoefficients)
+							{
+								Writer->WriteValue(InstallTimeCoefficient);
+							}
+							Writer->WriteArrayEnd();
 						}
-						Writer->WriteArrayEnd();
 					}
 					Writer->WriteObjectEnd();
 				}
@@ -970,7 +1049,7 @@ namespace BuildPatchServices
 		{
 			if (Download->ResponseSuccessful())
 			{
-				Async(EAsyncExecution::ThreadPool, [Download, RelevantPromisePtr]()
+				UE::Tasks::Launch(UE_SOURCE_LOCATION, [Download, RelevantPromisePtr]()
 				{
 					FBuildPatchAppManifestPtr Manifest = MakeShareable(new FBuildPatchAppManifest());
 					if (!Manifest->DeserializeFromData(Download->GetData()))
@@ -985,10 +1064,14 @@ namespace BuildPatchServices
 				RelevantPromisePtr->SetValue(FBuildPatchAppManifestPtr());
 			}
 		}
+		else
+		{
+			UE_LOG(LogDiffManifests, Error, TEXT("HandleDownloadComplete RequestId %d not tracked."), RequestId);
+		}
 	}
 
 	IDiffManifests* FDiffManifestsFactory::Create(const FDiffManifestsConfiguration& Configuration)
 	{
 		return new FDiffManifests(Configuration);
 	}
-}
+}
\ No newline at end of file
diff --git a/Engine/Source/Runtime/Online/BuildPatchServices/Private/Enumeration/PatchDataEnumeration.cpp b/Engine/Source/Runtime/Online/BuildPatchServices/Private/Enumeration/PatchDataEnumeration.cpp
--- a/Engine/Source/Runtime/Online/BuildPatchServices/Private/Enumeration/PatchDataEnumeration.cpp
+++ b/Engine/Source/Runtime/Online/BuildPatchServices/Private/Enumeration/PatchDataEnumeration.cpp
@@ -8,9 +8,10 @@
 #include "Common/FileSystem.h"
 #include "Data/ChunkData.h"
 #include "BuildPatchManifest.h"
+#include "BuildPatchMergeManifests.h"
 #include "BuildPatchUtil.h"
 
-DECLARE_LOG_CATEGORY_CLASS(LogDataEnumeration, Log, All);
+DEFINE_LOG_CATEGORY(LogDataEnumeration);
 
 namespace EnumerationHelpers
 {
@@ -77,7 +78,7 @@ namespace BuildPatchServices
 	{
 	public:
 
-		FPatchDataEnumeration(const FPatchDataEnumerationConfiguration& InConfiguration);
+		FPatchDataEnumeration(FPatchDataEnumerationConfiguration InConfiguration);
 		~FPatchDataEnumeration();
 
 		// IPatchDataEnumeration interface begin.
@@ -96,8 +97,8 @@ namespace BuildPatchServices
 		TUniquePtr<IFileSystem> FileSystem;
 	};
 
-	FPatchDataEnumeration::FPatchDataEnumeration(const FPatchDataEnumerationConfiguration& InConfiguration)
-		: Configuration(InConfiguration)
+	FPatchDataEnumeration::FPatchDataEnumeration(FPatchDataEnumerationConfiguration InConfiguration)
+		: Configuration(MoveTemp(InConfiguration))
 		, CloudDir(FPaths::GetPath(Configuration.InputFile))
 		, FileSystem(FFileSystemFactory::Create())
 	{
@@ -136,6 +137,10 @@ namespace BuildPatchServices
 		{
 			bSuccess = RunInternal(*File, FullOutputList);
 		}
+		else
+		{
+			UE_LOG(LogDataEnumeration, Error, TEXT("Failed to open the file %s"), *Configuration.InputFile);
+		}
 		return bSuccess;
 	}
 
@@ -200,10 +205,18 @@ namespace BuildPatchServices
 				EnumerationHelpers::AppendManifestDataList(*Manifest.Get(), OutFiles, VisitedDataSet, Configuration.bIncludeSizes);
 				for (const TUniquePtr<FBuildPatchAppManifest>& DeltaFile : DeltaFiles)
 				{
-					EnumerationHelpers::AppendManifestDataList(*DeltaFile.Get(), OutFiles, VisitedDataSet, Configuration.bIncludeSizes);
+					FBuildPatchAppManifestPtr MergedManifest = FBuildMergeManifests::MergeDeltaManifest(*Manifest.Get(), *DeltaFile.Get());
+					if (MergedManifest.IsValid())
+					{
+						EnumerationHelpers::AppendManifestDataList(*MergedManifest.Get(), OutFiles, VisitedDataSet, Configuration.bIncludeSizes);
+					}
 				}
 			}
 		}
+		else
+		{
+			UE_LOG(LogDataEnumeration, Error, TEXT("Failed to extract a manifest data from archive"));
+		}
 		return bSuccess;
 	}
 
@@ -244,11 +257,15 @@ namespace BuildPatchServices
 				}
 			}
 		}
+		else
+		{
+			UE_LOG(LogDataEnumeration, Error, TEXT("Archive serialization failed because of internal error"));
+		}
 		return bSuccess;
 	}
 
-	IPatchDataEnumeration* FPatchDataEnumerationFactory::Create(const FPatchDataEnumerationConfiguration& Configuration)
+	IPatchDataEnumeration* FPatchDataEnumerationFactory::Create(FPatchDataEnumerationConfiguration Configuration)
 	{
-		return new FPatchDataEnumeration(Configuration);
+		return new FPatchDataEnumeration(MoveTemp(Configuration));
 	}
 }
diff --git a/Engine/Source/Runtime/Online/BuildPatchServices/Private/Enumeration/PatchDataEnumeration.h b/Engine/Source/Runtime/Online/BuildPatchServices/Private/Enumeration/PatchDataEnumeration.h
--- a/Engine/Source/Runtime/Online/BuildPatchServices/Private/Enumeration/PatchDataEnumeration.h
+++ b/Engine/Source/Runtime/Online/BuildPatchServices/Private/Enumeration/PatchDataEnumeration.h
@@ -6,11 +6,13 @@
 #include "Interfaces/IPatchDataEnumeration.h"
 #include "BuildPatchSettings.h"
 
+DECLARE_LOG_CATEGORY_EXTERN(LogDataEnumeration, Log, All);
+
 namespace BuildPatchServices
 {
 	class FPatchDataEnumerationFactory
 	{
 	public:
-		static IPatchDataEnumeration* Create(const FPatchDataEnumerationConfiguration& Configuration);
+		static IPatchDataEnumeration* Create(FPatchDataEnumerationConfiguration Configuration);
 	};
 }
diff --git a/Engine/Source/Runtime/Online/BuildPatchServices/Private/Generation/BuildStreamer.h b/Engine/Source/Runtime/Online/BuildPatchServices/Private/Generation/BuildStreamer.h
--- a/Engine/Source/Runtime/Online/BuildPatchServices/Private/Generation/BuildStreamer.h
+++ b/Engine/Source/Runtime/Online/BuildPatchServices/Private/Generation/BuildStreamer.h
@@ -1,17 +1,16 @@
 // Copyright Epic Games, Inc. All Rights Reserved.
 #pragma once
 
-#include "CoreMinimal.h"
-#include "Misc/SecureHash.h"
-#include "HAL/FileManager.h"
-
+#include "BuildPatchManifest.h"
+#include "Common/StatsCollector.h"
 #include "Core/BlockStructure.h"
 #include "Core/Factory.h"
 #include "Core/FileSpan.h"
-#include "Common/StatsCollector.h"
+#include "CoreMinimal.h"
+#include "HAL/FileManager.h"
 #include "Installer/ChunkReferenceTracker.h"
 #include "Installer/CloudChunkSource.h"
-#include "BuildPatchManifest.h"
+#include "Misc/SecureHash.h"
 
 namespace BuildPatchServices
 {
@@ -40,27 +39,6 @@ namespace BuildPatchServices
 		: public IBuildStreamer
 	{
 	public:
-		/**
-		 * Retrieves the file details for a specific start index.
-		 * @param   IN  StartingIdx     The data index into the build image.
-		 * @param   OUT FileSpan        Receives a copy of the file span data.
-		 * @return true if the data byte at StartingIdx is the start of a file, false indicates that FileSpan was not set.
-		 */
-		virtual bool GetFileSpan(uint64 StartingIdx, FFileSpan& FileSpan) const = 0;
-
-		/**
-		 * Gets a list of empty files that the build contains.
-		 * @return array of empty files in the build.
-		 */
-		virtual TArray<FString> GetEmptyFiles() const = 0;
-
-		/**
-		 * Gets a list of all filenames that the build contains.
-		 * Will block until the list of files is enumerated and ignored files have been stripped out.
-		 * @return array of filenames in the build.
-		 */
-		virtual TArray<FString> GetAllFilenames() const = 0;
-
 		/**
 		 * Get the total build size that was streamed.
 		 * MUST be called only after IsEndOfData returns true.
@@ -103,8 +81,11 @@ namespace BuildPatchServices
 	{
 	public:
 		const FString BuildRoot;
-		const FString InputListFile;
-		const FString IgnoreListFile;
+		const TArray<FString> BuildFiles;
+		const bool bFollowSymlinks = false;
+		const bool bDetectMimeType = true;
+		const bool bCalculateMD5 = true;
+		const bool bCalculateSHA256 = true;
 	};
 
 	// Holds all dependencies for constructing a directory build streamer.
@@ -119,7 +100,7 @@ namespace BuildPatchServices
 	struct FManifestBuildStreamerConfig
 	{
 	public:
-		FString CloudDirectory;
+		TArray<FString> CloudDirectories;
 		FBlockStructure DesiredBytes;
 	};
 
diff --git a/Engine/Source/Runtime/Online/BuildPatchServices/Private/Generation/ChunkDatabaseWriter.cpp b/Engine/Source/Runtime/Online/BuildPatchServices/Private/Generation/ChunkDatabaseWriter.cpp
--- a/Engine/Source/Runtime/Online/BuildPatchServices/Private/Generation/ChunkDatabaseWriter.cpp
+++ b/Engine/Source/Runtime/Online/BuildPatchServices/Private/Generation/ChunkDatabaseWriter.cpp
@@ -9,9 +9,11 @@
 #include "Core/AsyncHelpers.h"
 #include "Data/ChunkData.h"
 #include "Common/FileSystem.h"
+#include "Common/StatsCollector.h"
 #include "Installer/InstallerError.h"
 #include "Installer/ChunkReferenceTracker.h"
 #include "Installer/ChunkSource.h"
+#include "Installer/Statistics/FileConstructorStatistics.h"
 #include "Interfaces/IBuildInstaller.h"
 
 DECLARE_LOG_CATEGORY_EXTERN(LogChunkDatabaseWriter, Log, All);
@@ -22,6 +24,16 @@ namespace BuildPatchServices
 	// Using initial data buffer of 2MB.
 	static const int32 DataMessageBufferSize = 1024 * 1024 * 2;
 
+	uint64 CountTotalChunks(const FChunkDbWriterConfig& Configuration)
+	{
+		uint64 NumDataMessages = 0;
+		for (const FChunkDatabaseFile& ChunkDatabaseFile : Configuration.ChunkDatabaseList)
+		{
+			NumDataMessages += ChunkDatabaseFile.DataList.Num();
+		}
+		return NumDataMessages;
+	}
+
 	struct FDataMessage
 	{
 	public:
@@ -62,38 +74,65 @@ namespace BuildPatchServices
 		: public IChunkDatabaseWriter
 	{
 	public:
-		FChunkDatabaseWriter(IChunkSource* ChunkSource, IFileSystem* FileSystem, IInstallerError* InstallerError, IChunkReferenceTracker* ChunkReferenceTracker, IChunkDataSerialization* ChunkDataSerialization, TArray<FChunkDatabaseFile> ChunkDatabaseList, TFunction<void(bool)> OnComplete);
+		FChunkDatabaseWriter(FChunkDbWriterConfig Configuration, IChunkSource& ChunkSource, IFileSystem& FileSystem, IInstallerError& InstallerError, IChunkReferenceTracker& ChunkReferenceTracker, IChunkDataSerialization& ChunkDataSerialization, IFileConstructorStat& FileConstructorStat, TFunction<void(float)> OnProgress, TFunction<void(bool)> OnComplete);
 
 		~FChunkDatabaseWriter();
 
+		// IControllable interface begin.
+		virtual void SetPaused(bool bInIsPaused) override
+		{
+			bIsPaused = bInIsPaused;
+		}
+
+		virtual void Abort() override
+		{
+			bShouldCancel = true;
+		}
+
+		virtual void Reset() override
+		{
+			bShouldCancel = false;
+		}
+		// IControllable interface end.
+
 	private:
 		void ProcessingWorkerThread();
 		void OutputWorkerThread();
+		void HandleProgress(float Progress);
+		void HandleComplete(bool bSuccess);
 
 	private:
-		IChunkSource* ChunkSource;
-		IFileSystem* FileSystem;
-		IInstallerError* InstallerError;
-		IChunkReferenceTracker* ChunkReferenceTracker;
-		IChunkDataSerialization* ChunkDataSerialization;
-		TArray<FChunkDatabaseFile> ChunkDatabaseList;
-		TFunction<void(bool)> OnComplete;
+		const FChunkDbWriterConfig Configuration;
+		const uint64 TotalNumChunks;
+		IChunkSource& ChunkSource;
+		IFileSystem& FileSystem;
+		IInstallerError& InstallerError;
+		IChunkReferenceTracker& ChunkReferenceTracker;
+		IChunkDataSerialization& ChunkDataSerialization;
+		IFileConstructorStat& FileConstructorStat;
+		TFunction<void(float)> OnProgressFn;
+		TFunction<void(bool)> OnCompleteFn;
 		TFuture<void> ProcessingWorkerFuture;
 		TFuture<void> OutputWorkerFuture;
+		FThreadSafeBool bIsPaused;
 		FThreadSafeBool bShouldCancel;
 		FThreadSafeBool bProcessingComplete;
 		TQueue<TSharedPtr<FDataMessage, ESPMode::ThreadSafe>, EQueueMode::Spsc> DataPipe;
+		FThreadSafeCounter MessageCount;
 		FEvent* ThreadTrigger;
 	};
 
-	FChunkDatabaseWriter::FChunkDatabaseWriter(IChunkSource* InChunkSource, IFileSystem* InFileSystem, IInstallerError* InInstallerError, IChunkReferenceTracker* InChunkReferenceTracker, IChunkDataSerialization* InChunkDataSerialization, TArray<FChunkDatabaseFile> InChunkDatabaseList, TFunction<void(bool)> InOnComplete)
-		: ChunkSource(InChunkSource)
+	FChunkDatabaseWriter::FChunkDatabaseWriter(FChunkDbWriterConfig InConfiguration, IChunkSource& InChunkSource, IFileSystem& InFileSystem, IInstallerError& InInstallerError, IChunkReferenceTracker& InChunkReferenceTracker, IChunkDataSerialization& InChunkDataSerialization, IFileConstructorStat& InFileConstructorStat, TFunction<void(float)> InOnProgress, TFunction<void(bool)> InOnComplete)
+		: Configuration(InConfiguration)
+		, TotalNumChunks(CountTotalChunks(Configuration))
+		, ChunkSource(InChunkSource)
 		, FileSystem(InFileSystem)
 		, InstallerError(InInstallerError)
 		, ChunkReferenceTracker(InChunkReferenceTracker)
 		, ChunkDataSerialization(InChunkDataSerialization)
-		, ChunkDatabaseList(MoveTemp(InChunkDatabaseList))
-		, OnComplete(MoveTemp(InOnComplete))
+		, FileConstructorStat(InFileConstructorStat)
+		, OnProgressFn(MoveTemp(InOnProgress))
+		, OnCompleteFn(MoveTemp(InOnComplete))
 		, bShouldCancel(false)
 		, bProcessingComplete(false)
 	{
@@ -121,14 +160,15 @@ namespace BuildPatchServices
 		bool bSuccess = true;
 
 		// For every entry in the provided ChunkDatabaseList, create the chunkdb, and send serialized data to the output thread for it.
-		for (int32 ChunkDatabaseIdx = 0; ChunkDatabaseIdx < ChunkDatabaseList.Num() && bSuccess && !bShouldCancel && !InstallerError->HasError(); ++ChunkDatabaseIdx)
+		for (int32 ChunkDatabaseIdx = 0; ChunkDatabaseIdx < Configuration.ChunkDatabaseList.Num() && bSuccess && !bShouldCancel && !InstallerError.HasError(); ++ChunkDatabaseIdx)
 		{
-			const FChunkDatabaseFile& ChunkDatabaseFile = ChunkDatabaseList[ChunkDatabaseIdx];
+			const FChunkDatabaseFile& ChunkDatabaseFile = Configuration.ChunkDatabaseList[ChunkDatabaseIdx];
 			UE_LOG(LogChunkDatabaseWriter, Log, TEXT("Start processing chunk database %s"), *ChunkDatabaseFile.DatabaseFilename);
 
 			// Send file create message.
-			const FString TmpDatabaseFilename = ChunkDatabaseFile.DatabaseFilename + TEXT("tmp");
-			DataPipe.Enqueue(MakeShareable(new FDataMessage(TmpDatabaseFilename, FDataMessage::CreateFileId)));
+			const FString CreateFilename = ChunkDatabaseFile.DatabaseFilename + (Configuration.bUseTempFile ? TEXT("tmp") : TEXT(""));
+			DataPipe.Enqueue(MakeShareable(new FDataMessage(CreateFilename, FDataMessage::CreateFileId)));
+			MessageCount.Increment();
 			ThreadTrigger->Trigger();
 
 			// Populate header with all required entries.
@@ -137,21 +177,32 @@ namespace BuildPatchServices
 			{
 				return FChunkLocation{DataId, 0, 0};
 			});
-			int64 FileDataPos = 0;
 
-			// Write the header.
-			TUniquePtr<FDataMessage> DataMessage(new FDataMessage(FileDataPos));
-			TUniquePtr<FMemoryWriter> MemoryWriter(new FMemoryWriter(DataMessage->Memory));
-			*MemoryWriter << ChunkDbHeader;
-			FileDataPos += DataMessage->Memory.Num();
-			DataPipe.Enqueue(MakeShareable(DataMessage.Release()));
-			ThreadTrigger->Trigger();
+			// Write the initial header.
+			TUniquePtr<FDataMessage> DataMessage;
+			TUniquePtr<FMemoryWriter> MemoryWriter;
+			uint64 LastHeaderWrite = 0;
+			int64 FileDataPos = 0;
+			auto SendHeaderMessage = [&]()
+			{
+				DataMessage.Reset(new FDataMessage(0));
+				MemoryWriter.Reset(new FMemoryWriter(DataMessage->Memory));
+				ChunkDbHeader.DataSize = FileDataPos - ChunkDbHeader.HeaderSize;
+				*MemoryWriter << ChunkDbHeader;
+				int64 NumBytesSent = DataMessage->Memory.Num();
+				DataPipe.Enqueue(MakeShareable(DataMessage.Release()));
+				MessageCount.Increment();
+				ThreadTrigger->Trigger();
+				LastHeaderWrite = FStatsCollector::GetCycles();
+				return NumBytesSent;
+			};
+			FileDataPos = SendHeaderMessage();
 
 			// Serialize and write each of the chunk files.
-			for (int32 ChunkDataIdx = 0; ChunkDataIdx < ChunkDatabaseFile.DataList.Num() && bSuccess && !bShouldCancel && !InstallerError->HasError(); ChunkDataIdx++)
+			for (int32 ChunkDataIdx = 0; ChunkDataIdx < ChunkDatabaseFile.DataList.Num() && bSuccess && !bShouldCancel && !InstallerError.HasError(); ChunkDataIdx++)
 			{
 				const FGuid& ChunkDataId = ChunkDatabaseFile.DataList[ChunkDataIdx];
-				IChunkDataAccess* ChunkDataAccess = ChunkSource->Get(ChunkDataId);
+				IChunkDataAccess* ChunkDataAccess = ChunkSource.Get(ChunkDataId);
 				bSuccess = ChunkDataAccess != nullptr;
 				if (bSuccess)
 				{
@@ -160,14 +211,25 @@ namespace BuildPatchServices
 					MemoryWriter.Reset(new FMemoryWriter(DataMessage->Memory));
 
 					// Write to message.
-					EChunkSaveResult SaveResult = ChunkDataSerialization->SaveToArchive(*MemoryWriter, ChunkDataAccess);
+					EChunkSaveResult SaveResult;
+					if (Configuration.bReserialise)
+					{
+						SaveResult = ChunkDataSerialization.SaveToArchive(*MemoryWriter, ChunkDataAccess);
+					}
+					else
+					{
+						FScopeLockedChunkData LockedChunkData(ChunkDataAccess);
+						*MemoryWriter << *LockedChunkData.GetHeader();
+						MemoryWriter->Serialize(LockedChunkData.GetData(), LockedChunkData.GetHeader()->DataSizeCompressed);
+						SaveResult = MemoryWriter->IsError() ? EChunkSaveResult::SerializationError : EChunkSaveResult::Success;
+					}
 					bSuccess = SaveResult == EChunkSaveResult::Success;
 					if (!bSuccess)
 					{
 						const TCHAR* ErrorCode = SaveResult == EChunkSaveResult::FileCreateFail ? ConstructionErrorCodes::FileCreateFail
 						                       : SaveResult == EChunkSaveResult::SerializationError ? ConstructionErrorCodes::SerializationError
 						                       : ConstructionErrorCodes::UnknownFail;
-						InstallerError->SetError(EBuildPatchInstallError::FileConstructionFail, ErrorCode);
+						InstallerError.SetError(EBuildPatchInstallError::FileConstructionFail, ErrorCode);
 					}
 
 					// Set the positional data in the header.
@@ -178,38 +240,58 @@ namespace BuildPatchServices
 					// Advance file position.
 					FileDataPos += Location.ByteSize;
 
+					// Wait while behind
+					while (MessageCount.GetValue() > Configuration.BufferEnqueueMax && !bShouldCancel)
+					{
+						FPlatformProcess::Sleep(0.1f);
+					}
+
 					// Send the data message.
 					DataPipe.Enqueue(MakeShareable(DataMessage.Release()));
+					MessageCount.Increment();
 					ThreadTrigger->Trigger();
 
 					// Pop the chunk we just saved out.
-					bSuccess = ChunkReferenceTracker->PopReference(ChunkDataId);
+					bSuccess = ChunkReferenceTracker.PopReference(ChunkDataId);
 					if (!bSuccess)
 					{
-						InstallerError->SetError(EBuildPatchInstallError::InitializationError, InitializationErrorCodes::ChunkReferenceTracking);
+						InstallerError.SetError(EBuildPatchInstallError::InitializationError, InitializationErrorCodes::ChunkReferenceTracking);
+					}
+					else if (Configuration.HeaderUpdateFrequency > 0.0f)
+					{
+						// Check if we should update the header to disk.
+						uint64 CurrentCycles = FStatsCollector::GetCycles();
+						if (Configuration.HeaderUpdateFrequency <= FStatsCollector::CyclesToSeconds(CurrentCycles - LastHeaderWrite))
+						{
+							SendHeaderMessage();
+						}
 					}
 				}
 				if (!bSuccess)
 				{
 					UE_LOG(LogChunkDatabaseWriter, Error, TEXT("    Failed chunk %s"), *ChunkDataId.ToString());
 					// We set a catch all error here, which only applies if our own loop above, or the chunk source, did not already set its specific error condition.
-					InstallerError->SetError(EBuildPatchInstallError::FileConstructionFail, ConstructionErrorCodes::UnknownFail);
+					InstallerError.SetError(EBuildPatchInstallError::FileConstructionFail, ConstructionErrorCodes::UnknownFail);
+				}
+				// Wait while paused
+				while (bIsPaused && !bShouldCancel)
+				{
+					FPlatformProcess::Sleep(0.5f);
 				}
 			}
 			if (bSuccess)
 			{
 				// Write back the header with all chunk positions now filled out accurately.
-				ChunkDbHeader.DataSize = FileDataPos - ChunkDbHeader.HeaderSize;
-				DataMessage.Reset(new FDataMessage(0));
-				MemoryWriter.Reset(new FMemoryWriter(DataMessage->Memory));
-				*MemoryWriter << ChunkDbHeader;
-				DataPipe.Enqueue(MakeShareable(DataMessage.Release()));
-				ThreadTrigger->Trigger();
+				SendHeaderMessage();
 			}
 
 			// Send file rename message.
-			DataPipe.Enqueue(MakeShareable(new FDataMessage(ChunkDatabaseFile.DatabaseFilename, FDataMessage::RenameFileId)));
-			ThreadTrigger->Trigger();
+			if (Configuration.bUseTempFile)
+			{
+				DataPipe.Enqueue(MakeShareable(new FDataMessage(ChunkDatabaseFile.DatabaseFilename, FDataMessage::RenameFileId)));
+				MessageCount.Increment();
+				ThreadTrigger->Trigger();
+			}
 		}
 
 		// Mark completed.
@@ -221,26 +303,39 @@ namespace BuildPatchServices
 	void FChunkDatabaseWriter::OutputWorkerThread()
 	{
 		bool bSuccess = true;
-
+		const float TotalNumChunksFloat = TotalNumChunks;
+		float ProcessedNumChunksFloat = 0;
 		TArray<FString> FilesCreated;
 		TSharedPtr<FDataMessage, ESPMode::ThreadSafe> DataMessage;
 		TUniquePtr<FArchive> CurrentFile;
-		while (bSuccess && !bShouldCancel && !InstallerError->HasError())
+		while (bSuccess && !bShouldCancel && !InstallerError.HasError())
 		{
 			// See if we have a message.
 			if (DataPipe.Dequeue(DataMessage))
 			{
+				MessageCount.Decrement();
 				// Process a file create message.
 				if (DataMessage->DataInfo == FDataMessage::CreateFileId)
 				{
+					if (CurrentFile.IsValid())
+					{
+						FileConstructorStat.OnFileCompleted(CurrentFile->GetArchiveName(), !CurrentFile->IsError());
+					}
 					UE_LOG(LogChunkDatabaseWriter, Log, TEXT("Writing chunk database %s"), *DataMessage->Filename);
-					CurrentFile = FileSystem->CreateFileWriter(*DataMessage->Filename);
+					ISpeedRecorder::FRecord ActivityRecord;
+					FileConstructorStat.OnFileStarted(DataMessage->Filename, 0);
+					FileConstructorStat.OnBeforeAdminister();
+					ActivityRecord.CyclesStart = FStatsCollector::GetCycles();
+					CurrentFile = FileSystem.CreateFileWriter(*DataMessage->Filename);
+					ActivityRecord.CyclesEnd = FStatsCollector::GetCycles();
+					ActivityRecord.Size = 0;
+					FileConstructorStat.OnAfterAdminister(ActivityRecord);
 					FilesCreated.Add(DataMessage->Filename);
 					if (CurrentFile.IsValid() == false)
 					{
 						bSuccess = false;
 						UE_LOG(LogChunkDatabaseWriter, Error, TEXT("Failed to create file with name %s"), *DataMessage->Filename);
-						InstallerError->SetError(EBuildPatchInstallError::FileConstructionFail, ConstructionErrorCodes::FileCreateFail);
+						InstallerError.SetError(EBuildPatchInstallError::FileConstructionFail, ConstructionErrorCodes::FileCreateFail);
 					}
 				}
 				// Process a data serialize.
@@ -248,9 +343,24 @@ namespace BuildPatchServices
 				{
 					if (CurrentFile->Tell() != DataMessage->DataInfo)
 					{
+						ISpeedRecorder::FRecord ActivityRecord;
+						FileConstructorStat.OnBeforeAdminister();
+						ActivityRecord.CyclesStart = FStatsCollector::GetCycles();
 						CurrentFile->Seek(DataMessage->DataInfo);
+						ActivityRecord.CyclesEnd = FStatsCollector::GetCycles();
+						ActivityRecord.Size = 0;
+						FileConstructorStat.OnAfterAdminister(ActivityRecord);
 					}
+					ISpeedRecorder::FRecord ActivityRecord;
+					FileConstructorStat.OnBeforeWrite();
+					ActivityRecord.CyclesStart = FStatsCollector::GetCycles();
 					CurrentFile->Serialize(DataMessage->Memory.GetData(), DataMessage->Memory.Num());
+					ActivityRecord.Size = DataMessage->Memory.Num();
+					ActivityRecord.CyclesEnd = FStatsCollector::GetCycles();
+					FileConstructorStat.OnAfterWrite(ActivityRecord);
+					FileConstructorStat.OnFileProgress(CurrentFile->GetArchiveName(), CurrentFile->TotalSize());
+					ProcessedNumChunksFloat += 1.0f;
+					HandleProgress(ProcessedNumChunksFloat / TotalNumChunksFloat);
 				}
 				// Process a file rename message.
 				else if (CurrentFile.IsValid() && DataMessage->DataInfo == FDataMessage::RenameFileId)
@@ -258,15 +368,22 @@ namespace BuildPatchServices
 					const FString OldFilename = CurrentFile->GetArchiveName();
 					bSuccess = CurrentFile->Close();
 					CurrentFile.Reset();
+					FileConstructorStat.OnFileCompleted(OldFilename, bSuccess);
 					if (bSuccess)
 					{
-						FileSystem->MoveFile(*DataMessage->Filename, *OldFilename);
+						ISpeedRecorder::FRecord ActivityRecord;
+						FileConstructorStat.OnBeforeAdminister();
+						ActivityRecord.CyclesStart = FStatsCollector::GetCycles();
+						FileSystem.MoveFile(*DataMessage->Filename, *OldFilename);
+						ActivityRecord.CyclesEnd = FStatsCollector::GetCycles();
+						ActivityRecord.Size = 0;
+						FileConstructorStat.OnAfterAdminister(ActivityRecord);
 						UE_LOG(LogChunkDatabaseWriter, Log, TEXT("Chunk database complete, renamed %s"), *DataMessage->Filename);
 					}
 					else
 					{
 						UE_LOG(LogChunkDatabaseWriter, Error, TEXT("Serialisation error reported on file close %s"), *OldFilename);
-						InstallerError->SetError(EBuildPatchInstallError::FileConstructionFail, ConstructionErrorCodes::SerializationError);
+						InstallerError.SetError(EBuildPatchInstallError::FileConstructionFail, ConstructionErrorCodes::SerializationError);
 					}
 				}
 				// An error if we do not have a file open and we were sent any message other than create.
@@ -274,7 +391,7 @@ namespace BuildPatchServices
 				{
 					bSuccess = false;
 					UE_LOG(LogChunkDatabaseWriter, Error, TEXT("Output fail, message without a file"));
-					InstallerError->SetError(EBuildPatchInstallError::FileConstructionFail, ConstructionErrorCodes::MissingFileInfo);
+					InstallerError.SetError(EBuildPatchInstallError::FileConstructionFail, ConstructionErrorCodes::MissingFileInfo);
 				}
 			}
 			// Quit if no more messages
@@ -291,33 +408,62 @@ namespace BuildPatchServices
 		}
 
 		// Close the last open file.
+		if (CurrentFile.IsValid())
+		{
+			FileConstructorStat.OnFileCompleted(CurrentFile->GetArchiveName(), !CurrentFile->IsError());
+		}
 		CurrentFile.Reset();
 
 		// Check whether the process was canceled or an error occurred.
-		bSuccess = bSuccess && !bShouldCancel && !InstallerError->HasError();
+		bSuccess = bSuccess && !bShouldCancel && !InstallerError.HasError();
 		UE_LOG(LogChunkDatabaseWriter, Log, TEXT("Writer complete! bSuccess:%d"), bSuccess);
 
 		// Delete any created files if we failed.
-		if (!bSuccess)
+		if (!bSuccess && Configuration.bDeleteFilesOnFailure)
 		{
 			UE_LOG(LogChunkDatabaseWriter, Error, TEXT("Chunkdb generation failed. All created files will be deleted."));
 			for (const FString& FileToDelete : FilesCreated)
 			{
-				FileSystem->DeleteFile(*FileToDelete);
+				FileSystem.DeleteFile(*FileToDelete);
 				UE_LOG(LogChunkDatabaseWriter, Log, TEXT("Deleted %s"), *FileToDelete);
 			}
 		}
 
 		// We're done so call the complete callback.
-		AsyncHelpers::ExecuteOnGameThread(OnComplete, bSuccess).Wait();
+		HandleComplete(bSuccess);
+	}
+
+	void FChunkDatabaseWriter::HandleProgress(float Progress)
+	{
+		if (Configuration.bCallbackOnMainThread)
+		{
+			AsyncHelpers::ExecuteOnGameThread(OnProgressFn, Progress).Wait();
+		}
+		else
+		{
+			OnProgressFn(Progress);
+		}
+	}
+
+	void FChunkDatabaseWriter::HandleComplete(bool bSuccess)
+	{
+		if (Configuration.bCallbackOnMainThread)
+		{
+			AsyncHelpers::ExecuteOnGameThread(OnCompleteFn, bSuccess).Wait();
+		}
+		else
+		{
+			OnCompleteFn(bSuccess);
+		}
 	}
 
-	IChunkDatabaseWriter* FChunkDatabaseWriterFactory::Create(IChunkSource* ChunkSource, IFileSystem* FileSystem, IInstallerError* InstallerError, IChunkReferenceTracker* ChunkReferenceTracker, IChunkDataSerialization* ChunkDataSerialization, TArray<FChunkDatabaseFile> ChunkDatabaseList, TFunction<void(bool)> OnComplete)
+	IChunkDatabaseWriter* FChunkDatabaseWriterFactory::Create(FChunkDbWriterConfig Configuration, IChunkSource* ChunkSource, IFileSystem* FileSystem, IInstallerError* InstallerError, IChunkReferenceTracker* ChunkReferenceTracker, IChunkDataSerialization* ChunkDataSerialization, IFileConstructorStat* FileConstructorStat, TFunction<void(float)> OnProgress, TFunction<void(bool)> OnComplete)
 	{
 		check(ChunkSource != nullptr);
 		check(FileSystem != nullptr);
 		check(InstallerError != nullptr);
 		check(ChunkReferenceTracker != nullptr);
-		return new FChunkDatabaseWriter(ChunkSource, FileSystem, InstallerError, ChunkReferenceTracker, ChunkDataSerialization, MoveTemp(ChunkDatabaseList), MoveTemp(OnComplete));
+		check(FileConstructorStat != nullptr);
+		return new FChunkDatabaseWriter(MoveTemp(Configuration), *ChunkSource, *FileSystem, *InstallerError, *ChunkReferenceTracker, *ChunkDataSerialization, *FileConstructorStat, MoveTemp(OnProgress), MoveTemp(OnComplete));
 	}
 }
diff --git a/Engine/Source/Runtime/Online/BuildPatchServices/Private/Generation/ChunkDatabaseWriter.h b/Engine/Source/Runtime/Online/BuildPatchServices/Private/Generation/ChunkDatabaseWriter.h
--- a/Engine/Source/Runtime/Online/BuildPatchServices/Private/Generation/ChunkDatabaseWriter.h
+++ b/Engine/Source/Runtime/Online/BuildPatchServices/Private/Generation/ChunkDatabaseWriter.h
@@ -2,6 +2,7 @@
 #pragma once
 
 #include "CoreMinimal.h"
+#include "Installer/Controllable.h"
 
 namespace BuildPatchServices
 {
@@ -10,6 +11,7 @@ namespace BuildPatchServices
 	class IInstallerError;
 	class IChunkReferenceTracker;
 	class IChunkDataSerialization;
+	class IFileConstructorStat;
 
 	struct FChunkDatabaseFile
 	{
@@ -22,11 +24,44 @@ namespace BuildPatchServices
 	 * into databases.
 	 */
 	class IChunkDatabaseWriter
+		: public IControllable
 	{
 	public:
 		virtual ~IChunkDatabaseWriter() = default;
 	};
 
+	/**
+	 * A struct containing the configuration values for a chunkdb writer.
+	 */
+	struct FChunkDbWriterConfig
+	{
+		// The array of chunk database files to create and the chunks to place in them.
+		TArray<FChunkDatabaseFile> ChunkDatabaseList;
+		// The max threaded messages to wait on in case of HDD write is slower that chunk fetch.
+		int32 BufferEnqueueMax = 100;
+		// Whether to use temp file when building, this is safer against failures, but if false, can support resuming.
+		// It is advised to use true when generating for cloud upload, and false if producing for local use only on a client for later consumption.
+		bool bUseTempFile = true;
+		// If more than 0, the header will be updated every X seconds to minimise loss of data if the process is cancelled.
+		float HeaderUpdateFrequency = 0.0f;
+		// If true, will potentially re-encrypt and/or re-compress on save to chumdb.
+		// If false, the chunk will be stored as deserialised, and thus decrypted/decompressed where it was possible.
+		bool bReserialise = true;
+		// If true, will delete created files upon a failure, otherwise will leave potentially unfinished files on disk.
+		bool bDeleteFilesOnFailure = true;
+		// If set to false, the progress and complete callbacks will be executed on the worker threads.
+		bool bCallbackOnMainThread = true;
+
+		/**
+		 * Constructor which sets usual defaults, and takes params for values that cannot use a default.
+		 * @param InChunkDatabaseList    The chunkdb filename array.
+		 */
+		FChunkDbWriterConfig(TArray<FChunkDatabaseFile> InChunkDatabaseList)
+			: ChunkDatabaseList(MoveTemp(InChunkDatabaseList))
+		{
+		}
+	};
+
 	/**
 	 * A factory for creating an IChunkDatabaseWriter instance.
 	 */
@@ -35,22 +70,26 @@ namespace BuildPatchServices
 	public:
 		/**
 		 * This implementation returns a chunk database writer that immediately kicks off the work and calls a provided callback when complete.
+		 * @param Configuration             The configuration struct for this source.
 		 * @param ChunkSource               A chunk source for pulling required chunks from.
 		 * @param FileSystem                A files system interface for writing out the chunkdb files.
 		 * @param InstallerError            The error interface for aborting on other errors or registering our own.
 		 * @param ChunkReferenceTracker     Chunk reference tracker to keep up to date.
 		 * @param ChunkDataSerialization    Chunk data serialization implementation.
-		 * @param ChunkDatabaseList         The array of chunk database files to create and the chunks to place in them.
+		 * @param InFileConstructorStat     A file constructor stats class to track HDD activity.
+		 * @param OnProgress                Function to call with a progress, 0->1. Called on main thread.
 		 * @param OnComplete                Function to call when the database files have been created. Called on main thread.
 		 * @return the new IChunkDatabaseWriter.
 		 */
 		static IChunkDatabaseWriter* Create(
+			FChunkDbWriterConfig Configuration,
 			IChunkSource* ChunkSource,
 			IFileSystem* FileSystem,
 			IInstallerError* InstallerError,
 			IChunkReferenceTracker* ChunkReferenceTracker,
 			IChunkDataSerialization* ChunkDataSerialization,
-			TArray<FChunkDatabaseFile> ChunkDatabaseList,
+			IFileConstructorStat* InFileConstructorStat,
+			TFunction<void(float)> OnProgress,
 			TFunction<void(bool)> OnComplete);
 	};
 }
