p4 integrate //UE5/Dev-Distro-5.5/Engine/Source/Programs/BuildPatchTool/... //UE5/Main/Engine/Source/Programs/BuildPatchTool/...

diff --git a/Engine/Source/Runtime/Online/BuildPatchServices/Private/Generation/ChunkDeltaOptimiser.cpp b/Engine/Source/Runtime/Online/BuildPatchServices/Private/Generation/ChunkDeltaOptimiser.cpp
--- a/Engine/Source/Runtime/Online/BuildPatchServices/Private/Generation/ChunkDeltaOptimiser.cpp
+++ b/Engine/Source/Runtime/Online/BuildPatchServices/Private/Generation/ChunkDeltaOptimiser.cpp
@@ -2,46 +2,52 @@
 
 #include "Generation/ChunkDeltaOptimiser.h"
 
-#include "Containers/List.h"
-#include "Containers/Ticker.h"
-#include "Containers/Queue.h"
-#include "HAL/ThreadSafeBool.h"
 #include "Async/Async.h"
-#include "Misc/ConfigCacheIni.h"
-#include "Misc/FileHelper.h"
-#include "Misc/OutputDeviceRedirector.h"
-#include "Policies/CondensedJsonPrintPolicy.h"
-#include "Serialization/JsonWriter.h"
-#include "HttpModule.h"
-
-#include "Core/MeanValue.h"
-#include "Core/Platform.h"
-#include "Core/ProcessTimer.h"
+#include "BuildPatchHash.h"
+#include "BuildPatchServicesModule.h"
+#include "BuildPatchUtil.h"
 #include "Common/ChunkDataSizeProvider.h"
-#include "Common/SpeedRecorder.h"
+#include "Common/Crypto.h"
 #include "Common/FileSystem.h"
 #include "Common/HttpManager.h"
+#include "Common/SpeedRecorder.h"
 #include "Common/StatsCollector.h"
+#include "Containers/List.h"
+#include "Containers/Queue.h"
+#include "Containers/Ticker.h"
+#include "Core/AsyncHelpers.h"
+#include "Core/MeanValue.h"
+#include "Core/Platform.h"
+#include "Core/ProcessTimer.h"
 #include "Generation/BuildStreamer.h"
 #include "Generation/ChunkMatchProcessor.h"
-#include "Generation/ChunkWriter.h"
 #include "Generation/ChunkSearch.h"
+#include "Generation/ChunkWriter.h"
 #include "Generation/DataScanner.h"
 #include "Generation/DeltaEnumeration.h"
+#include "HAL/ThreadSafeBool.h"
+#include "HttpModule.h"
+#include "IBuildManifestSet.h"
+#include "Installer/ChunkEvictionPolicy.h"
 #include "Installer/InstallerAnalytics.h"
-#include "Installer/Statistics/DownloadServiceStatistics.h"
+#include "Installer/InstallerError.h"
+#include "Installer/InstallerSharedContext.h"
 #include "Installer/MemoryChunkStore.h"
-#include "Installer/ChunkEvictionPolicy.h"
 #include "Installer/MessagePump.h"
-#include "Installer/InstallerError.h"
+#include "Installer/OptimisedDelta.h"
+#include "Installer/UriProvider.h"
 #include "Installer/Statistics/CloudChunkSourceStatistics.h"
-#include "BuildPatchHash.h"
-#include "BuildPatchUtil.h"
-#include "IBuildManifestSet.h"
-
+#include "Installer/Statistics/CloudChunkSourceStat.h"
+#include "Installer/Statistics/DownloadServiceStatistics.h"
 #include "Misc/CommandLine.h"
+#include "Misc/ConfigCacheIni.h"
+#include "Misc/FileHelper.h"
+#include "Misc/OutputDeviceRedirector.h"
+#include "Policies/CondensedJsonPrintPolicy.h"
+#include "Serialization/JsonWriter.h"
 
-DECLARE_LOG_CATEGORY_CLASS(LogChunkDeltaOptimiser, Log, All);
+DECLARE_LOG_CATEGORY_EXTERN(LogChunkDeltaOptimiser, Log, All);
+DEFINE_LOG_CATEGORY(LogChunkDeltaOptimiser);
 
 // For the output file we'll use pretty json in debug, otherwise condensed.
 #if UE_BUILD_DEBUG
@@ -175,26 +181,94 @@ namespace DeltaOptimiseHelpers
 		});
 	}
 
-	void MakeScannerLocalList(FChunkSearcher& ChunkSearcher, IDeltaChunkEnumeration* Enumeration, const FBlockStructure& BuildStructure, FScannerFilesList& Result)
+	void MakeScannerLocalList(const TMap<FString, FString>& AliasedTags, FChunkSearcher& ChunkSearcher, IDeltaChunkEnumeration* Enumeration, const FBlockStructure& BuildStructure, FScannerFilesList& Result)
 	{
 		uint64 FirstByte = 0;
 		ChunkSearcher.ForEachOverlap(BuildStructure, [&](const FBlockRange& OverlapRange, FChunkSearcher::FFileDListNode* File, FChunkSearcher::FChunkDListNode* Chunk)
 		{
 			const FFilenameId FilenameId = Enumeration->MakeFilenameId(File->GetValue().Manifest->Filename);
-			const TArray<FString>& FileTagset = File->GetValue().Manifest->InstallTags;
 			const FBlockRange& FileRange = File->GetValue().BuildRange;
 			const uint64 FileOffset = OverlapRange.GetFirst() - FileRange.GetFirst();
-			Result.AddTail(FScannerFileElement{FBlockRange::FromFirstAndSize(FirstByte, OverlapRange.GetSize()), FilenameId, FileTagset, FileOffset});
+
+			// Replace any aliased tags.
+			TSet<FString> FileTagset;
+			FileTagset.Reserve(File->GetValue().Manifest->InstallTags.Num());
+			for (const FString& InstallTag : File->GetValue().Manifest->InstallTags)
+			{
+				const FString* AliasedTag = AliasedTags.Find(InstallTag);
+				FileTagset.Add(AliasedTag ? *AliasedTag : InstallTag);
+			}
+
+			Result.AddTail(FScannerFileElement{FBlockRange::FromFirstAndSize(FirstByte, OverlapRange.GetSize()), FilenameId, MoveTemp(FileTagset), FileOffset});
 			FirstByte += OverlapRange.GetSize();
 		});
 		check(BlockStructureHelpers::CountSize(BuildStructure) == FirstByte);
 	}
+
+	uint64 AsPercentageOf(const BuildPatchServices::FDiffAbortThreshold& DiffAbortThreshold, uint64 TotalBinarySize)
+	{
+		check(DiffAbortThreshold.Unit == BuildPatchServices::EDiffAbortThresholdUnits::Absolute);
+		return TotalBinarySize == 0 ? 100 : DiffAbortThreshold.Value * 100 / TotalBinarySize;
+	}
+
+	uint64 AsAbsoluteOf(const BuildPatchServices::FDiffAbortThreshold& DiffAbortThreshold, uint64 TotalBinarySize)
+	{
+		check(DiffAbortThreshold.Unit == EDiffAbortThresholdUnits::Percentage);
+		return TotalBinarySize * DiffAbortThreshold.Value / 100;
+	}
+
+	BuildPatchServices::FDiffAbortThreshold ClampToSaneMinBasedOnActualBinarySize(TArray<FString>& FinalStatLogs, const FDiffAbortThreshold& DiffAbortThreshold, uint64 TotalBinarySize)
+	{
+		if (DiffAbortThreshold.Unit == EDiffAbortThresholdUnits::Absolute)
+		{
+			check(DiffAbortThreshold.Value >= DiffAbortThresholdLimits::MinAbsolute); // it must be clamped previously
+			return DiffAbortThreshold;
+		}
+
+		check(DiffAbortThreshold.Unit == EDiffAbortThresholdUnits::Percentage);
+
+		uint64 MaxPercentageOfTotalSize = AsPercentageOf({ DiffAbortThresholdLimits::MinAbsolute, EDiffAbortThresholdUnits::Absolute }, TotalBinarySize);
+		MaxPercentageOfTotalSize = FMath::Max(MaxPercentageOfTotalSize, DiffAbortThresholdLimits::MinPercentage);
+		if (DiffAbortThreshold.Value < MaxPercentageOfTotalSize)
+		{
+			// Clamp DiffAbortThreshold to sane min.
+			FinalStatLogs.Add(FString::Printf(TEXT("Requested -DiffAbortThreshold=%llu%% is below allowed minimum n >= 1GB. Please update your arg to be above limit. Continuing with DiffAbortThreshold=%llu%%."),
+				DiffAbortThreshold.Value, MaxPercentageOfTotalSize));
+			return { MaxPercentageOfTotalSize, EDiffAbortThresholdUnits::Percentage };
+		}
+		return DiffAbortThreshold;
+	}
+
+	struct FAbsoluteAndPercentage
+	{
+		uint64 Absolute = 0;
+		uint64 Percentage = 0;
+	};
+
+	FAbsoluteAndPercentage RecalculateAsAbsoluteAndPercentage(const FDiffAbortThreshold& DiffAbortThreshold, uint64 TotalBinarySize)
+	{
+		switch (DiffAbortThreshold.Unit)
+		{
+		case EDiffAbortThresholdUnits::Absolute:
+			return { DiffAbortThreshold.Value, AsPercentageOf(DiffAbortThreshold, TotalBinarySize) };
+		case EDiffAbortThresholdUnits::Percentage:
+			return { AsAbsoluteOf(DiffAbortThreshold, TotalBinarySize), DiffAbortThreshold.Value};
+		default:
+			unimplemented();
+			break;
+		}
+		return {};
+	}
+
+
 }
 
 namespace DeltaStats
 {
+	using namespace BuildPatchServices;
+
 	class FNoMemoryChunkStoreStat
-		: public BuildPatchServices::IMemoryChunkStoreStat
+		: public IMemoryChunkStoreStat
 	{
 	public:
 		FNoMemoryChunkStoreStat() { }
@@ -208,28 +282,6 @@ namespace DeltaStats
 		virtual void OnStoreSizeUpdated(int32 Size) override { }
 		// IMemoryChunkStoreStat interface end.
 	};
-
-	class FNoCloudChunkSourceStat
-		: public BuildPatchServices::ICloudChunkSourceStat
-	{
-	public:
-		FNoCloudChunkSourceStat() { }
-		~FNoCloudChunkSourceStat() { }
-
-		// ICloudChunkSourceStat interface begin.
-		virtual void OnDownloadRequested(const FGuid& ChunkId) override { }
-		virtual void OnDownloadSuccess(const FGuid& ChunkId) override { }
-		virtual void OnDownloadFailed(const FGuid& ChunkId, const FString& Url) override { }
-		virtual void OnDownloadCorrupt(const FGuid& ChunkId, const FString& Url, BuildPatchServices::EChunkLoadResult LoadResult) override { }
-		virtual void OnDownloadAborted(const FGuid& ChunkId, const FString& Url, double DownloadTimeMean, double DownloadTimeStd, double DownloadTime, double BreakingPoint) override { }
-		virtual void OnReceivedDataUpdated(int64 TotalBytes) override { }
-		virtual void OnRequiredDataUpdated(int64 TotalBytes) override { }
-		virtual void OnDownloadHealthUpdated(EBuildPatchDownloadHealth DownloadHealth) override { }
-		virtual void OnSuccessRateUpdated(float SuccessRate) override { }
-		virtual void OnActiveRequestCountUpdated(uint32 RequestCount) override { }
-		virtual void OnAcceptedNewRequirements(const TSet<FGuid>& ChunkIds) override { }
-		// ICloudChunkSourceStat interface end.
-	};
 }
 
 namespace DeltaFactories
@@ -258,12 +310,13 @@ namespace DeltaFactories
 		IChunkDataSerialization* ChunkDataSerialization;
 		IMessagePump* MessagePump;
 		IBuildManifestSet* ManifestSet;
+		ICloudChunkSourceStat* CloudChunkSourceStat;
 	};
 
 	class FCloudChunkSourceFactory : public IManifestBuildStreamer::ICloudChunkSourceFactory
 	{
 	private:
-		struct FInstanceDependancies
+		struct FInstanceDependencies
 		{
 			TUniquePtr<IChunkEvictionPolicy> MemoryEvictionPolicy;
 			TUniquePtr<IMemoryChunkStore> CloudChunkStore;
@@ -271,16 +324,23 @@ namespace DeltaFactories
 		};
 
 	public:
-		FCloudChunkSourceFactory(const FString& CloudDir, FCloudChunkSourceFactoryShared InShared)
+		FCloudChunkSourceFactory(TArray<FString> CloudDirs, FCloudChunkSourceFactoryShared InShared)
 			: Shared(MoveTemp(InShared))
-			, CloudSourceConfig({CloudDir})
+			, CloudSourceSharedContext(FBuildInstallerSharedContextFactory::Create(TEXT("CloudChunkSourceFactory")))
+			, CloudSourceConfig(MoveTemp(CloudDirs))
 			, Platform(FPlatformFactory::Create())
 			, MemoryChunkStoreStat(new DeltaStats::FNoMemoryChunkStoreStat())
 			, InstallerError(FInstallerErrorFactory::Create())
-			, CloudChunkSourceStat(new DeltaStats::FNoCloudChunkSourceStat())
 		{
+			CloudSourceSharedContext->PreallocateThreads(1);
+
 			CloudSourceConfig.bBeginDownloadsOnFirstGet = false;
 			CloudSourceConfig.MaxRetryCount = 30;
+			GConfig->GetInt(TEXT("Portal.BuildPatch"), TEXT("ChunkRetries"), CloudSourceConfig.MaxRetryCount, GEngineIni);
+			CloudSourceConfig.RetriesToEnableVerbose = CloudSourceConfig.MaxRetryCount - 1;
+			// Lower the prefetch max as the default is a bit high for our use case here
+			CloudSourceConfig.PreFetchMaximum = 100;
+			CloudSourceConfig.SharedContext = &CloudSourceSharedContext.Get();
 		}
 
 		virtual ~FCloudChunkSourceFactory()
@@ -290,27 +350,28 @@ namespace DeltaFactories
 		// IManifestBuildStreamer::ICloudChunkSourceFactory interface begin.
 		virtual ICloudChunkSource* Create(IChunkReferenceTracker* ChunkReferenceTracker) override
 		{
-			InstanceDependancies.AddDefaulted();
-			FInstanceDependancies& Dependancies = InstanceDependancies.Last();
-			Dependancies.MemoryEvictionPolicy.Reset(FChunkEvictionPolicyFactory::Create(ChunkReferenceTracker));
-			Dependancies.CloudChunkStore.Reset(FMemoryChunkStoreFactory::Create(
-				100,
-				Dependancies.MemoryEvictionPolicy.Get(),
+			InstanceDependencies.AddDefaulted();
+			FInstanceDependencies& Dependencies = InstanceDependencies.Last();
+			Dependencies.MemoryEvictionPolicy.Reset(FChunkEvictionPolicyFactory::Create(ChunkReferenceTracker));
+			Dependencies.CloudChunkStore.Reset(FMemoryChunkStoreFactory::Create(
+				// The chunk store should be bigger than the max prefetch to allow for multi-use chunks to stay in memory.
+				CloudSourceConfig.PreFetchMaximum + 50,
+				Dependencies.MemoryEvictionPolicy.Get(),
 				nullptr,
 				MemoryChunkStoreStat.Get(),
 				nullptr));
-			Dependancies.ConnectionCount.Reset(BuildPatchServices::FDownloadConnectionCountFactory::Create(ConnectionCountConfig, nullptr));
+			Dependencies.ConnectionCount.Reset(BuildPatchServices::FDownloadConnectionCountFactory::Create(ConnectionCountConfig, nullptr));
 			ICloudChunkSource* CloudChunkSource = BuildPatchServices::FCloudChunkSourceFactory::Create(
 				CloudSourceConfig,
 				Platform.Get(),
-				Dependancies.CloudChunkStore.Get(),
+				Dependencies.CloudChunkStore.Get(),
 				Shared.DownloadService,
 				ChunkReferenceTracker,
 				Shared.ChunkDataSerialization,
 				Shared.MessagePump,
 				InstallerError.Get(),
-				Dependancies.ConnectionCount.Get(),
-				CloudChunkSourceStat.Get(),
+				Dependencies.ConnectionCount.Get(),
+				Shared.CloudChunkSourceStat,
 				Shared.ManifestSet,
 				ChunkReferenceTracker->GetReferencedChunks());
 
@@ -318,7 +379,7 @@ namespace DeltaFactories
 			{
 				CloudChunkSource->AddRepeatRequirement(LostChunk);
 			};
-			Dependancies.CloudChunkStore->SetLostChunkCallback(LostChunkCallback);
+			Dependencies.CloudChunkStore->SetLostChunkCallback(LostChunkCallback);
 
 			return CloudChunkSource;
 		}
@@ -326,13 +387,13 @@ namespace DeltaFactories
 
 	private:
 		FCloudChunkSourceFactoryShared Shared;
+		IBuildInstallerSharedContextRef CloudSourceSharedContext;
 		FCloudSourceConfig CloudSourceConfig;
 		FDownloadConnectionCountConfig ConnectionCountConfig;
 		TUniquePtr<IPlatform> Platform;
 		TUniquePtr<IMemoryChunkStoreStat> MemoryChunkStoreStat;
 		TUniquePtr<IInstallerError> InstallerError;
-		TUniquePtr<ICloudChunkSourceStat> CloudChunkSourceStat;
-		TArray<FInstanceDependancies> InstanceDependancies;
+		TArray<FInstanceDependencies> InstanceDependencies;
 	};
 }
 
@@ -563,65 +624,76 @@ namespace BuildPatchServices
 		: public IChunkDeltaOptimiser
 	{
 	public:
-		FChunkDeltaOptimiser(const FChunkDeltaOptimiserConfiguration& InConfiguration);
+		FChunkDeltaOptimiser(FChunkDeltaOptimiserConfiguration InConfiguration);
 		~FChunkDeltaOptimiser();
 
 		// IChunkDeltaOptimiser interface begin.
 		virtual bool Run() override;
+		virtual FOnChunkFileWritten& OnChunkFileWritten() override;
+		virtual FOnDeltaFileWritten& OnDeltaFileWritten() override;
 		// IChunkDeltaOptimiser interface end.
 
+	private:
+		enum class EFetchResult
+		{
+			Success = 0,
+			DoesNotExist,
+			CouldNotDeserialize
+		};
+		typedef TTuple<FBuildPatchAppManifestPtr, int32, EFetchResult> FFetchResult;
+
 	private:
 		TArray<FString> AsyncRun();
-		void HandleDownloadComplete(int32 RequestId, const FDownloadRef& Download);
+		void BroadcastChunkFile(const FString& FullFilePath, const FMD5Hash& MD5Hash);
+		void BroadcastDelta(const FString& FullFilePath, const FMD5Hash& MD5Hash, const FSHAHash& SHA1Hash);
+		TFuture<FFetchResult> FetchManifest(const FString& ManifestUri);
+		void HandleDownloadComplete(int32 RequestId, const FDownloadRef& Download, TSharedRef<TPromise<FFetchResult>, ESPMode::ThreadSafe> Promise);
 		FBlockStructure GetDesiredBytes(const FBuildPatchAppManifestPtr& Manifest, const TSet<FGuid>& Chunks);
+		FOptimisedDeltaConfiguration BuildOptimisedDeltaConfig(FBuildPatchAppManifestPtr SourceManifest, FBuildPatchAppManifestRef DestinationManifest, const TArray<FString>& DownloadCloudDirectories, const FString& FilenameTrailer);
+		FOptimisedDeltaDependencies BuildOptimisedDeltaDependencies(TSharedRef<IUriProvider, ESPMode::ThreadSafe> UriProvider);
 
 	private:
 		const FChunkDeltaOptimiserConfiguration Configuration;
 		FTSTicker& CoreTicker;
 		FDownloadCompleteDelegate DownloadCompleteDelegate;
 		FDownloadProgressDelegate DownloadProgressDelegate;
-		TUniquePtr<IFileSystem> FileSystem;
-		TUniquePtr<IHttpManager> HttpManager;
+		TSharedPtr<IFileSystem> FileSystem;
+		TUniquePtr<ICrypto> Crypto;
+		TSharedPtr<IHttpManager> HttpManager;
 		TUniquePtr<IChunkDataSizeProvider> ChunkDataSizeProvider;
 		TUniquePtr<ISpeedRecorder> DownloadSpeedRecorder;
-		TUniquePtr<IInstallerAnalytics> InstallerAnalytics;
-		TUniquePtr<IDownloadServiceStatistics> DownloadServiceStatistics;
+		TSharedPtr<IInstallerAnalytics> InstallerAnalytics;
+		TSharedPtr<IDownloadServiceStatistics> DownloadServiceStatistics;
 		TUniquePtr<IDownloadService> DownloadService;
-		TUniquePtr<IMessagePump> MessagePump;
+		TUniquePtr<IMessagePump> MessagePumpA;
+		TUniquePtr<IMessagePump> MessagePumpB;
 		TUniquePtr<FStatsCollector> StatsCollector;
+		TArray<FMessageHandler*> MessageHandlersA;
+		TArray<FMessageHandler*> MessageHandlersB;
 		FThreadSafeBool bShouldRun;
 		FThreadSafeBool bSuccess;
-
-		// Manifest downloading
-		int32 RequestIdManifestA;
-		int32 RequestIdManifestB;
-		TPromise<FBuildPatchAppManifestPtr> PromiseManifestA;
-		TPromise<FBuildPatchAppManifestPtr> PromiseManifestB;
-		TFuture<FBuildPatchAppManifestPtr> FutureManifestA;
-		TFuture<FBuildPatchAppManifestPtr> FutureManifestB;
+		IChunkDeltaOptimiser::FOnChunkFileWritten ChunkFileWrittenEvent;
+		IChunkDeltaOptimiser::FOnDeltaFileWritten DeltaFileWrittenEvent;
 	};
 
-	FChunkDeltaOptimiser::FChunkDeltaOptimiser(const FChunkDeltaOptimiserConfiguration& InConfiguration)
-		: Configuration(InConfiguration)
+	FChunkDeltaOptimiser::FChunkDeltaOptimiser(FChunkDeltaOptimiserConfiguration InConfiguration)
+		: Configuration(MoveTemp(InConfiguration))
 		, CoreTicker(FTSTicker::GetCoreTicker())
-		, DownloadCompleteDelegate(FDownloadCompleteDelegate::CreateRaw(this, &FChunkDeltaOptimiser::HandleDownloadComplete))
 		, DownloadProgressDelegate()
 		, FileSystem(FFileSystemFactory::Create())
+		, Crypto(FCryptoFactory::Create())
 		, HttpManager(FHttpManagerFactory::Create())
-		, ChunkDataSizeProvider(FChunkDataSizeProviderFactory::Create())
+		, ChunkDataSizeProvider(FChunkDataSizeProviderFactory::Create({}))
 		, DownloadSpeedRecorder(FSpeedRecorderFactory::Create())
 		, InstallerAnalytics(FInstallerAnalyticsFactory::Create(nullptr))
 		, DownloadServiceStatistics(FDownloadServiceStatisticsFactory::Create(DownloadSpeedRecorder.Get(), ChunkDataSizeProvider.Get(), InstallerAnalytics.Get()))
-		, DownloadService(FDownloadServiceFactory::Create(HttpManager.Get(), FileSystem.Get(), DownloadServiceStatistics.Get(), InstallerAnalytics.Get()))
-		, MessagePump(FMessagePumpFactory::Create())
+		, DownloadService(FDownloadServiceFactory::Create(HttpManager, FileSystem, DownloadServiceStatistics, InstallerAnalytics))
+		, MessagePumpA(FMessagePumpFactory::Create())
+		, MessagePumpB(FMessagePumpFactory::Create())
 		, StatsCollector(FStatsCollectorFactory::Create())
+		, MessageHandlersA(Configuration.MessageHandlersA)
+		, MessageHandlersB(Configuration.MessageHandlersB)
 		, bShouldRun(true)
-		, RequestIdManifestA(INDEX_NONE)
-		, RequestIdManifestB(INDEX_NONE)
-		, PromiseManifestA()
-		, PromiseManifestB()
-		, FutureManifestA(PromiseManifestA.GetFuture())
-		, FutureManifestB(PromiseManifestB.GetFuture())
 	{
 	}
 
@@ -635,14 +707,12 @@ namespace BuildPatchServices
 		FHttpModule::Get();
 
 		// Setup Generation stats.
-		volatile int64* StatTotalTime = StatsCollector->CreateStat(TEXT("Generation: Total Time"), EStatFormat::Timer);
-		volatile int64* StatStreamSpeed = StatsCollector->CreateStat(TEXT("Generation: Data stream speed"), EStatFormat::DataSpeed);
+		volatile FStatsCollector::FAtomicValue* StatTotalTime = StatsCollector->CreateStat(TEXT("Generation: Total Time"), EStatFormat::Timer);
+		volatile FStatsCollector::FAtomicValue* StatAverageStreamSpeed = StatsCollector->CreateStat(TEXT("Generation: Average stream speed"), EStatFormat::DataSpeed);
+		volatile FStatsCollector::FAtomicValue* StatCurrentStreamSpeed = StatsCollector->CreateStat(TEXT("Generation: Current stream speed"), EStatFormat::DataSpeed);
 		const uint64 StartTime = FStatsCollector::GetCycles();
-		const float SpeedStatOverTime = TNumericLimits<float>::Max();
-
-		// Kick off Manifest downloads.
-		RequestIdManifestA = DownloadService->RequestFile(Configuration.ManifestAUri, DownloadCompleteDelegate, DownloadProgressDelegate);
-		RequestIdManifestB = DownloadService->RequestFile(Configuration.ManifestBUri, DownloadCompleteDelegate, DownloadProgressDelegate);
+		const float SpeedStatOverTimeAverage = TNumericLimits<float>::Max();
+		const float SpeedStatOverTimeCurrent = 10.f;
 
 		// Start the generation thread.
 		TFuture<TArray<FString>> Thread = Async(EAsyncExecution::Thread, [this](){ return AsyncRun(); });
@@ -671,12 +741,14 @@ namespace BuildPatchServices
 			FTSTicker::GetCoreTicker().Tick(DeltaTime);
 
 			// Message pump.
-			MessagePump->PumpMessages();
+			MessagePumpA->PumpMessages(MessageHandlersA);
+			MessagePumpB->PumpMessages(MessageHandlersB);
 
 			// Log collected stats.
 			GLog->FlushThreadedLogs();
 			FStatsCollector::Set(StatTotalTime, FStatsCollector::GetCycles() - StartTime);
-			FStatsCollector::Set(StatStreamSpeed, DownloadSpeedRecorder->GetAverageSpeed(SpeedStatOverTime));
+			FStatsCollector::Set(StatAverageStreamSpeed, DownloadSpeedRecorder->GetAverageSpeed(SpeedStatOverTimeAverage));
+			FStatsCollector::Set(StatCurrentStreamSpeed, DownloadSpeedRecorder->GetAverageSpeed(SpeedStatOverTimeCurrent));
 			StatsCollector->LogStats(StatsLoggerTimeSeconds);
 
 			// Control frame rate.
@@ -692,7 +764,8 @@ namespace BuildPatchServices
 		TArray<FString> FinalStatLogs = Thread.Get();
 		GLog->FlushThreadedLogs();
 		FStatsCollector::Set(StatTotalTime, FStatsCollector::GetCycles() - StartTime);
-		FStatsCollector::Set(StatStreamSpeed, DownloadSpeedRecorder->GetAverageSpeed(SpeedStatOverTime));
+		FStatsCollector::Set(StatAverageStreamSpeed, DownloadSpeedRecorder->GetAverageSpeed(SpeedStatOverTimeAverage));
+		FStatsCollector::Set(StatCurrentStreamSpeed, DownloadSpeedRecorder->GetAverageSpeed(SpeedStatOverTimeCurrent));
 		StatsCollector->LogStats();
 		for (const FString& LogLine : FinalStatLogs)
 		{
@@ -703,23 +776,76 @@ namespace BuildPatchServices
 		return bSuccess;
 	}
 
+	IChunkDeltaOptimiser::FOnChunkFileWritten& FChunkDeltaOptimiser::OnChunkFileWritten()
+	{
+		return ChunkFileWrittenEvent;
+	}
+
+	IChunkDeltaOptimiser::FOnDeltaFileWritten& FChunkDeltaOptimiser::OnDeltaFileWritten()
+	{
+		return DeltaFileWrittenEvent;
+	}
+
+	FOptimisedDeltaConfiguration FChunkDeltaOptimiser::BuildOptimisedDeltaConfig(FBuildPatchAppManifestPtr SourceManifest, FBuildPatchAppManifestRef DestinationManifest, const TArray<FString>& DownloadCloudDirectories, const FString& FilenameTrailer)
+	{
+		// The optimised delta can deal with getting dupe manifests so lets just allow that for easy config.
+		FOptimisedDeltaConfiguration OptimisedDeltaConfiguration(DestinationManifest);
+		OptimisedDeltaConfiguration.SourceManifest = SourceManifest;
+		OptimisedDeltaConfiguration.RetriesNumber = DownloadCloudDirectories.Num();
+		
+		// When the delta file can't be deserialized we don't want to proceed with re-uploading it.
+		// The file should be deleted on backend first, and then we can re-upload it.
+		OptimisedDeltaConfiguration.DeltaPolicy = EDeltaPolicy::Expect;
+		OptimisedDeltaConfiguration.DeltaFilenameTrailer = FilenameTrailer;
+		return OptimisedDeltaConfiguration;
+	}
+
+	FOptimisedDeltaDependencies FChunkDeltaOptimiser::BuildOptimisedDeltaDependencies(TSharedRef<IUriProvider, ESPMode::ThreadSafe> UriProvider)
+	{
+		FOptimisedDeltaDependencies OptimisedDeltaDependencies(MoveTemp(UriProvider));
+		OptimisedDeltaDependencies.DownloadService = DownloadService.Get();
+		return OptimisedDeltaDependencies;
+	}
+
 	TArray<FString> FChunkDeltaOptimiser::AsyncRun()
 	{
 		const FNumberFormattingOptions PercentFormat = FNumberFormattingOptions().SetMaximumFractionalDigits(1).SetMinimumFractionalDigits(1).SetRoundingMode(ERoundingMode::ToZero);
-		FBuildPatchAppManifestPtr ManifestA = FutureManifestA.Get();
-		FBuildPatchAppManifestPtr ManifestB = FutureManifestB.Get();
+		TFuture<FFetchResult> FutureManifestResultA = FetchManifest(Configuration.ManifestAUri);
+		TFuture<FFetchResult> FutureManifestResultB = FetchManifest(Configuration.ManifestBUri);
+		FFetchResult ManifestResultA = FutureManifestResultA.Get();
+		FFetchResult ManifestResultB = FutureManifestResultB.Get();
 		TArray<FString> FinalStatLogs;
 		bSuccess = true;
-		if (ManifestA.IsValid() == false)
+		if (ManifestResultA.Get<2>() != EFetchResult::Success)
 		{
 			UE_LOG(LogChunkDeltaOptimiser, Error, TEXT("Could not download ManifestA from %s."), *Configuration.ManifestAUri);
 			bSuccess = false;
 		}
-		if (ManifestB.IsValid() == false)
+		if (ManifestResultB.Get<2>() != EFetchResult::Success)
 		{
 			UE_LOG(LogChunkDeltaOptimiser, Error, TEXT("Could not download ManifestB from %s."), *Configuration.ManifestBUri);
 			bSuccess = false;
 		}
+		const FBuildPatchAppManifestPtr& ManifestA = ManifestResultA.Get<0>();
+		const FBuildPatchAppManifestPtr& ManifestB = ManifestResultB.Get<0>();
+		if (bSuccess)
+		{
+			TSet<FGuid> AvailableSecrets;
+			Configuration.EncryptionSecrets.GetKeys(AvailableSecrets);
+			TSet<FGuid> NecessaryEncryptionSecrets;
+			NecessaryEncryptionSecrets.Append(ManifestA->GetNecessaryEncryptionSecretIds());
+			NecessaryEncryptionSecrets.Append(ManifestB->GetNecessaryEncryptionSecretIds());
+			TSet<FGuid> MissingSecrets = NecessaryEncryptionSecrets.Difference(AvailableSecrets);
+			if (MissingSecrets.Num() > 0)
+			{
+				UE_LOG(LogChunkDeltaOptimiser, Error, TEXT("Cannot continue without all necessary secret keys."));
+				for (const FGuid& SecretId : MissingSecrets)
+				{
+					UE_LOG(LogChunkDeltaOptimiser, Error, TEXT("    Missing secret with ID: %s"), *SecretId.ToString());
+				}
+				bSuccess = false;
+			}
+		}
 		if (bSuccess)
 		{
 			UE_LOG(LogChunkDeltaOptimiser, Display, TEXT("Running optimisation for patching %s -> %s"), *ManifestA->GetVersionString(), *ManifestB->GetVersionString());
@@ -740,7 +866,7 @@ namespace BuildPatchServices
 			const bool bIsDowngrade = ManifestB->ManifestMeta.FeatureLevel < ManifestA->ManifestMeta.FeatureLevel;
 			if (bUsingDifferentChunkSubdir && bIsDowngrade)
 			{
-				UE_LOG(LogChunkDeltaOptimiser, Error, TEXT("Destination manifest does not support source manifest's FeatureLevel (%s [%d] -> %s [%d])."), FeatureLevelToString(ManifestA->ManifestMeta.FeatureLevel), (int32)ManifestA->ManifestMeta.FeatureLevel, FeatureLevelToString(ManifestB->ManifestMeta.FeatureLevel), (int32)ManifestB->ManifestMeta.FeatureLevel);
+				UE_LOG(LogChunkDeltaOptimiser, Error, TEXT("Destination manifest does not support source manifest's FeatureLevel (%s [%d] -> %s [%d])."), LexToString(ManifestA->ManifestMeta.FeatureLevel), (int32)ManifestA->ManifestMeta.FeatureLevel, LexToString(ManifestB->ManifestMeta.FeatureLevel), (int32)ManifestB->ManifestMeta.FeatureLevel);
 				bSuccess = false;
 			}
 
@@ -751,23 +877,37 @@ namespace BuildPatchServices
 				UE_LOG(LogChunkDeltaOptimiser, Log, TEXT("Destination manifest does not support EFeatureLevel::VariableSizeChunks, reverting OutputChunkSize to %u."), OutputChunkSize);
 			}
 
-			// Check that an optimisation does not already exist, and skip long process if so.
+			// Set cloud directories
+			TArray<FString> DownloadCloudDirectories(Configuration.DownloadCloudDirectories);
+			DownloadCloudDirectories.Add(Configuration.OutputCloudDirectory);
+
+			// Check that an optimisation does not already exist in any potential locations, and skip long process if so.
+			TSharedRef<IUriProvider, ESPMode::ThreadSafe> UriProvider = FUriProviderFactory::Create(MessagePumpB.Get(), Configuration.DownloadCloudDirectories);
+			FOptimisedDeltaDependencies OptimisedDeltaDependencies = BuildOptimisedDeltaDependencies(MoveTemp(UriProvider));
+			TSharedRef<IOptimisedDelta> OptimisedDelta = FOptimisedDeltaFactory::Create(BuildOptimisedDeltaConfig(ManifestA, ManifestB.ToSharedRef(), DownloadCloudDirectories, Configuration.DeltaFilenameTrailer), MoveTemp(OptimisedDeltaDependencies));
+
+			const FString RelativeDeltaFilename = FBuildPatchUtils::GetChunkDeltaFilename(*ManifestA.Get(), *ManifestB.Get(), Configuration.DeltaFilenameTrailer);
+			const FString OutputDeltaFilename = Configuration.OutputCloudDirectory / RelativeDeltaFilename;
+
+			const IOptimisedDelta::FResultValueOrError& OptimisedDeltaResult = OptimisedDelta->GetResult();
+
+			// Check whether the delta file was retrieved and can't be serialized
+			if (OptimisedDeltaResult.HasError() && OptimisedDeltaResult.GetError() == DownloadErrorCodes::UnserialisableDeltaFile)
+			{
+				UE_LOG(LogChunkDeltaOptimiser, Error, TEXT("Optimised delta completed previously but could not be loaded %s."), *RelativeDeltaFilename);
+				UE_LOG(LogChunkDeltaOptimiser, Error, TEXT("This file should be deleted as it will cause patching errors."));
+				bSuccess = false;
+			}
+
+			const bool bDeltaPreviouslyCompleted = OptimisedDeltaResult.IsValid() && OptimisedDeltaResult.GetValue() != ManifestB;
 			FBuildPatchAppManifest DeltaManifest;
-			// The Portal stream version of this has support for filename trailers, but we don't use it in this stream so we pass "".
-			const FString OutputDeltaFilename = Configuration.CloudDirectory / FBuildPatchUtils::GetChunkDeltaFilename(*ManifestA.Get(), *ManifestB.Get(), TEXT(""));
-			const bool bDeltaPreviouslyCompleted = FileSystem->FileExists(*OutputDeltaFilename);
-			if (bDeltaPreviouslyCompleted)
+			int32 LoadedDeltaSize = INDEX_NONE;
+			if (bSuccess && bDeltaPreviouslyCompleted)
 			{
-				if (DeltaManifest.LoadFromFile(OutputDeltaFilename))
-				{
-					FinalStatLogs.Add(FString::Printf(TEXT("** Chunk delta optimisation already completed for provided manifests. **")));
-					FinalStatLogs.Add(FString::Printf(TEXT("Loaded optimised delta file %s"), *OutputDeltaFilename));
-				}
-				else
-				{
-					UE_LOG(LogChunkDeltaOptimiser, Error, TEXT("Optimised delta completed previously but could not be loaded %s."), *OutputDeltaFilename);
-					bSuccess = false;
-				}
+				FinalStatLogs.Add(FString::Printf(TEXT("** Chunk delta optimisation already completed for provided manifests. **")));
+				FinalStatLogs.Add(FString::Printf(TEXT("Loaded optimised delta file %s"), *RelativeDeltaFilename));
+				DeltaManifest = *OptimisedDelta->GetResult().GetValue();
+				LoadedDeltaSize = OptimisedDelta->GetMetaDownloadSize();
 			}
 
 			// Check for aborting if original delta is over provided threshold.
@@ -775,18 +915,55 @@ namespace BuildPatchServices
 			ManifestA->GetFileTagList(TagsA);
 			ManifestB->GetFileTagList(TagsB);
 			const uint64 OriginalUnknownCompressedBytes = (uint64)ManifestB->GetDeltaDownloadSize(TagsB, ManifestA.ToSharedRef(), TagsA);
-			const bool bOverAbortThreshold = OriginalUnknownCompressedBytes >= Configuration.DiffAbortThreshold;
-			if (!bDeltaPreviouslyCompleted && bOverAbortThreshold)
+			bool bOverAbortThreshold = false;
+			if (Configuration.DiffAbortThreshold)
+			{
+				const FDiffAbortThreshold DiffAbortThreshold = DeltaOptimiseHelpers::ClampToSaneMinBasedOnActualBinarySize(FinalStatLogs, *Configuration.DiffAbortThreshold, ManifestB->GetBuildSize());
+				const DeltaOptimiseHelpers::FAbsoluteAndPercentage AbsoluteAndPercentage = DeltaOptimiseHelpers::RecalculateAsAbsoluteAndPercentage(DiffAbortThreshold, ManifestB->GetBuildSize());
+				bOverAbortThreshold = OriginalUnknownCompressedBytes >= AbsoluteAndPercentage.Absolute;
+				if (!bDeltaPreviouslyCompleted && bOverAbortThreshold)
+				{
+					const uint64 OriginalUnknownCompressedBytesInPercentage = DeltaOptimiseHelpers::AsPercentageOf({ OriginalUnknownCompressedBytes, EDiffAbortThresholdUnits::Absolute }, ManifestB->GetBuildSize());
+					FinalStatLogs.Add(FString::Printf(TEXT("** Aborting delta optimisation due to original delta over the %s threshold. **"), *LexToString(DiffAbortThreshold.Unit)));
+					FinalStatLogs.Add(FString::Printf(TEXT("%llu >= %llu Bytes (or in percentage %llu%% >= %llu%%)"), 
+						OriginalUnknownCompressedBytes, AbsoluteAndPercentage.Absolute, OriginalUnknownCompressedBytesInPercentage,  AbsoluteAndPercentage.Percentage));
+				}
+			}
+
+			// Check for aborting if original delta is less than window size.
+			const bool bIsLessThanWindowSize = OriginalUnknownCompressedBytes < (int32)Configuration.ScanWindowSize;
+			if (!bDeltaPreviouslyCompleted && bIsLessThanWindowSize)
+			{
+				FinalStatLogs.Add(FString::Printf(TEXT("** Aborting delta optimisation due to original delta less than scan window size. **")));
+				FinalStatLogs.Add(FString::Printf(TEXT("%llu < %u"), OriginalUnknownCompressedBytes, Configuration.ScanWindowSize));
+			}
+
+			// Calculate the desired bytes for manifest streams.
+			FBlockStructure ManifestADesiredBytes = GetDesiredBytes(ManifestA, ChunksA.Difference(ChunksB));
+			FBlockStructure ManifestBDesiredBytes = GetDesiredBytes(ManifestB, ChunksB.Difference(ChunksA));
+			const uint64 ManifestAStreamSize = BlockStructureHelpers::CountSize(ManifestADesiredBytes);
+			const uint64 ManifestBStreamSize = BlockStructureHelpers::CountSize(ManifestBDesiredBytes);
+
+			const bool bIsManifestASizeLessThanWindowsSize = ManifestAStreamSize < (int32)Configuration.ScanWindowSize;
+			if (!bDeltaPreviouslyCompleted && bIsManifestASizeLessThanWindowsSize)
+			{
+				FinalStatLogs.Add(FString::Printf(TEXT("** Aborting delta optimisation due to manifest A size less than scan window size. **")));
+				FinalStatLogs.Add(FString::Printf(TEXT("%llu < %u"), ManifestAStreamSize, Configuration.ScanWindowSize));
+			}
+
+			const bool bIsManifestBSizeLessThanWindowsSize = ManifestBStreamSize < (int32)Configuration.ScanWindowSize;
+			if (!bDeltaPreviouslyCompleted && bIsManifestBSizeLessThanWindowsSize)
 			{
-				FinalStatLogs.Add(FString::Printf(TEXT("** Aborting delta optimisation due to original delta over threshold. **")));
-				FinalStatLogs.Add(FString::Printf(TEXT("%llu >= %llu"), OriginalUnknownCompressedBytes, Configuration.DiffAbortThreshold));
+				FinalStatLogs.Add(FString::Printf(TEXT("** Aborting delta optimisation due to manifest B size less than scan window size. **")));
+				FinalStatLogs.Add(FString::Printf(TEXT("%llu < %u"), ManifestBStreamSize, Configuration.ScanWindowSize));
 			}
 
-			const bool bRunProcess = bSuccess && !bDeltaPreviouslyCompleted && !bOverAbortThreshold;
+			const bool bRunProcess = bSuccess && !bDeltaPreviouslyCompleted && !bOverAbortThreshold && !bIsLessThanWindowSize && !bIsManifestBSizeLessThanWindowsSize && !bIsManifestASizeLessThanWindowsSize;
 			if (bRunProcess)
 			{
 				// Runtime composition.
-				TUniquePtr<IChunkDataSerialization> ChunkDataSerializationReader(FChunkDataSerializationFactory::Create(FileSystem.Get()));
+				TUniquePtr<ICloudChunkSourceStat> CloudChunkSourceStat(FCloudChunkSourceStatFactory::Create(StatsCollector.Get()));
+				TUniquePtr<IChunkDataSerialization> ChunkDataSerializationReader(FChunkDataSerializationFactory::Create(FileSystem.Get(), Crypto.Get(), { ManifestB->ManifestMeta.FeatureLevel, Configuration.EncryptionSecrets, ManifestB->GetEncryptionSecretId() /* UE5 MERGE TODO : , CloudChunkSourceStat.Get()*/ }));
 				TUniquePtr<DeltaFactories::FChunkReferenceTrackerFactory> ChunkReferenceTrackerFactory(new DeltaFactories::FChunkReferenceTrackerFactory());
 				TUniquePtr<IBuildManifestSet> SetA(FBuildManifestSetFactory::Create({ FInstallerAction::MakeInstall(ManifestA.ToSharedRef()) }));
 				TUniquePtr<IBuildManifestSet> SetB(FBuildManifestSetFactory::Create({ FInstallerAction::MakeInstall(ManifestB.ToSharedRef()) }));
@@ -794,16 +971,18 @@ namespace BuildPatchServices
 				CloudChunkSourceFactorySharedA.FileSystem = FileSystem.Get();
 				CloudChunkSourceFactorySharedA.DownloadService = DownloadService.Get();
 				CloudChunkSourceFactorySharedA.ChunkDataSerialization = ChunkDataSerializationReader.Get();
-				CloudChunkSourceFactorySharedA.MessagePump = MessagePump.Get();
+				CloudChunkSourceFactorySharedA.MessagePump = MessagePumpA.Get();
 				CloudChunkSourceFactorySharedA.ManifestSet = SetA.Get();
+				CloudChunkSourceFactorySharedA.CloudChunkSourceStat = CloudChunkSourceStat.Get();
 				DeltaFactories::FCloudChunkSourceFactoryShared CloudChunkSourceFactorySharedB;
 				CloudChunkSourceFactorySharedB.FileSystem = FileSystem.Get();
 				CloudChunkSourceFactorySharedB.DownloadService = DownloadService.Get();
 				CloudChunkSourceFactorySharedB.ChunkDataSerialization = ChunkDataSerializationReader.Get();
-				CloudChunkSourceFactorySharedB.MessagePump = MessagePump.Get();
+				CloudChunkSourceFactorySharedB.MessagePump = MessagePumpB.Get();
 				CloudChunkSourceFactorySharedB.ManifestSet = SetB.Get();
-				TUniquePtr<DeltaFactories::FCloudChunkSourceFactory> CloudChunkSourceFactoryA(new DeltaFactories::FCloudChunkSourceFactory(Configuration.CloudDirectory, CloudChunkSourceFactorySharedA));
-				TUniquePtr<DeltaFactories::FCloudChunkSourceFactory> CloudChunkSourceFactoryB(new DeltaFactories::FCloudChunkSourceFactory(Configuration.CloudDirectory, CloudChunkSourceFactorySharedB));
+				CloudChunkSourceFactorySharedB.CloudChunkSourceStat = CloudChunkSourceStat.Get();
+				TUniquePtr<DeltaFactories::FCloudChunkSourceFactory> CloudChunkSourceFactoryA(new DeltaFactories::FCloudChunkSourceFactory(Configuration.DownloadCloudDirectories, CloudChunkSourceFactorySharedA));
+				TUniquePtr<DeltaFactories::FCloudChunkSourceFactory> CloudChunkSourceFactoryB(new DeltaFactories::FCloudChunkSourceFactory(Configuration.DownloadCloudDirectories, CloudChunkSourceFactorySharedB));
 
 				// Buffer for data streaming.
 				const EAllowShrinking AllowShrinking = EAllowShrinking::No;
@@ -812,30 +991,36 @@ namespace BuildPatchServices
 				TArray<uint8> StreamBuffer;
 				StreamBuffer.Reserve(StreamBufferReadSize + Configuration.ScanWindowSize);
 
-				// Calculate the desired bytes for manifest streams.
-				FBlockStructure ManifestADesiredBytes = GetDesiredBytes(ManifestA, ChunksA.Difference(ChunksB));
-				FBlockStructure ManifestBDesiredBytes = GetDesiredBytes(ManifestB, ChunksB.Difference(ChunksA));
-				const uint64 ManifestBStreamSize = BlockStructureHelpers::CountSize(ManifestBDesiredBytes);
-
 				// Start the ManifestA stream and chunk enumeration.
-				FManifestBuildStreamerConfig ManifestAStreamConfig({Configuration.CloudDirectory, ManifestADesiredBytes});
+				FManifestBuildStreamerConfig ManifestAStreamConfig({Configuration.DownloadCloudDirectories, ManifestADesiredBytes});
 				FManifestBuildStreamerDependencies ManifestAStreamDependencies({ChunkReferenceTrackerFactory.Get(), CloudChunkSourceFactoryA.Get(), StatsCollector.Get(), ManifestA.Get()});
 				TUniquePtr<IManifestBuildStreamer> ManifestAStream(FBuildStreamerFactory::Create(MoveTemp(ManifestAStreamConfig), MoveTemp(ManifestAStreamDependencies)));
 
 				// First we re-chunk prev unknown build parts into the scanner window size.
-				TUniquePtr<IDeltaChunkEnumeration> DeltaChunkEnumeration(FDeltaChunkEnumerationFactory::Create(ManifestAStream.Get(), StatsCollector.Get(), *ManifestA.Get(), Configuration.ScanWindowSize));
+				TUniquePtr<IDeltaChunkEnumeration> DeltaChunkEnumeration(FDeltaChunkEnumerationFactory::Create(ManifestAStream.Get(), StatsCollector.Get(), *ManifestA.Get(), Configuration.ScanWindowSize, Configuration.ManifestATagAliases));
 				ChunkingTimer.Start();
 				DeltaChunkEnumeration->Run();
 				ChunkingTimer.Stop();
 
+				if (!DeltaChunkEnumeration->IsSuccess())
+				{
+					bShouldRun = false;
+					bSuccess = false;
+					UE_LOG(LogChunkDeltaOptimiser, Error, TEXT("Failed to enumerate chunk data."));
+					return FinalStatLogs;
+				}
+
 				// Setup scanning stats.
 				volatile int64* StatScannerBacklog = StatsCollector->CreateStat(TEXT("BuildB: Scanner backlog"), EStatFormat::Value);
 				volatile int64* StatScannerForks = StatsCollector->CreateStat(TEXT("BuildB: Scanner forks"), EStatFormat::Value);
 				volatile int64* StatScanningTime = StatsCollector->CreateStat(TEXT("BuildB: Scanning time"), EStatFormat::Timer);
 				volatile int64* StatScanningCompleted = StatsCollector->CreateStat(TEXT("BuildB: Progress"), EStatFormat::Percentage);
 
+				// Place to dump details logging until we write it to a file at the end.
+				TArray<FString> SourceDetailsOutputLines;
+				
 				// Start the ManifestB stream.
-				FManifestBuildStreamerConfig ManifestBStreamConfig({Configuration.CloudDirectory, ManifestBDesiredBytes});
+				FManifestBuildStreamerConfig ManifestBStreamConfig({Configuration.DownloadCloudDirectories, ManifestBDesiredBytes});
 				FManifestBuildStreamerDependencies ManifestBStreamDependencies({ChunkReferenceTrackerFactory.Get(), CloudChunkSourceFactoryB.Get(), StatsCollector.Get(), ManifestB.Get()});
 				TUniquePtr<IManifestBuildStreamer> ManifestBStream(FBuildStreamerFactory::Create(MoveTemp(ManifestBStreamConfig), MoveTemp(ManifestBStreamDependencies)));
 
@@ -912,7 +1097,7 @@ namespace BuildPatchServices
 						for (FMatchEntry& AcceptedChunkMatch : AcceptedChunkMatches)
 						{
 							const FChunkMatch& ChunkCMatch = AcceptedChunkMatch.ChunkMatch;
-							const TArray<FChunkPart>& NewChunkReferences = ChunkBuildReferences[ChunkCMatch.ChunkGuid].Get<0>();
+							const TArray<FChunkPart>& NewChunkReferences = ChunkBuildReferences[ChunkCMatch.ChunkGuid].ChunkParts;
 							const FBlockStructure& ChunkCBuildBStructure = AcceptedChunkMatch.BlockStructure;
 							ChunkMatchStomper->ReplaceChunkReferences(NewChunkReferences, ChunkCBuildBStructure);
 						}
@@ -929,7 +1114,7 @@ namespace BuildPatchServices
 
 						FBlockStructure ScannerBuildStructure;
 						ManifestBDesiredBytes.SelectSerialBytes(NewScanner->Offset, SizeToScan, ScannerBuildStructure);
-						DeltaOptimiseHelpers::MakeScannerLocalList(FileListSearcher, DeltaChunkEnumeration.Get(), ScannerBuildStructure, NewScanner->FilesList);
+						DeltaOptimiseHelpers::MakeScannerLocalList(Configuration.ManifestBTagAliases, FileListSearcher, DeltaChunkEnumeration.Get(), ScannerBuildStructure, NewScanner->FilesList);
 
 						NewScanner->Scanner.Reset(FDeltaScannerFactory::Create(Configuration.ScanWindowSize, NewScanner->Data, NewScanner->FilesList, DeltaChunkEnumeration.Get(), StatsCollector.Get()));
 						ConsumedBufferData += SizeToScan;
@@ -967,7 +1152,7 @@ namespace BuildPatchServices
 
 								FBlockStructure ScannerBuildStructure;
 								ManifestBDesiredBytes.SelectSerialBytes(NewScanner->Offset, RightFork.GetSize(), ScannerBuildStructure);
-								DeltaOptimiseHelpers::MakeScannerLocalList(FileListSearcher, DeltaChunkEnumeration.Get(), ScannerBuildStructure, NewScanner->FilesList);
+								DeltaOptimiseHelpers::MakeScannerLocalList(Configuration.ManifestBTagAliases, FileListSearcher, DeltaChunkEnumeration.Get(), ScannerBuildStructure, NewScanner->FilesList);
 
 								NewScanner->Scanner.Reset(FDeltaScannerFactory::Create(Configuration.ScanWindowSize, NewScanner->Data, NewScanner->FilesList, DeltaChunkEnumeration.Get(), StatsCollector.Get()));
 								NewScanner->bIsFinalScanner = DataScannerEntry.bIsFinalScanner;
@@ -982,7 +1167,7 @@ namespace BuildPatchServices
 
 								ScannerBuildStructure.Empty();
 								ManifestBDesiredBytes.SelectSerialBytes(NewScanner->Offset, LeftFork.GetSize(), ScannerBuildStructure);
-								DeltaOptimiseHelpers::MakeScannerLocalList(FileListSearcher, DeltaChunkEnumeration.Get(), ScannerBuildStructure, NewScanner->FilesList);
+								DeltaOptimiseHelpers::MakeScannerLocalList(Configuration.ManifestBTagAliases, FileListSearcher, DeltaChunkEnumeration.Get(), ScannerBuildStructure, NewScanner->FilesList);
 
 								NewScanner->Scanner.Reset(FDeltaScannerFactory::Create(Configuration.ScanWindowSize, NewScanner->Data, NewScanner->FilesList, DeltaChunkEnumeration.Get(), StatsCollector.Get()));
 								NewScanner->bIsFinalScanner = false;
@@ -1025,7 +1210,7 @@ namespace BuildPatchServices
 						if (ChunkPart.Guid.IsValid() == false)
 						{
 							uint64 PartByteLocation = ByteLocation;
-							uint32 PartSizeRemaining = ChunkPart.Size;
+							uint64 PartSizeRemaining = ChunkPart.Size;
 							while (PartSizeRemaining > 0)
 							{
 								// Start new chunk?
@@ -1035,8 +1220,8 @@ namespace BuildPatchServices
 								}
 
 								TTuple<FBlockStructure, FChunkPart>* LastChunkDetail = &NewChunks.Last();
-								const uint32 NewTotalSize = LastChunkDetail->Get<1>().Size + PartSizeRemaining;
-								const uint32 ChunkPartConsume = NewTotalSize > OutputChunkSize ? PartSizeRemaining - (NewTotalSize - OutputChunkSize) : PartSizeRemaining;
+								const uint64 NewTotalSize = LastChunkDetail->Get<1>().Size + PartSizeRemaining;
+								const uint64 ChunkPartConsume = NewTotalSize > OutputChunkSize ? PartSizeRemaining - (NewTotalSize - OutputChunkSize) : PartSizeRemaining;
 								check(PartSizeRemaining >= ChunkPartConsume);
 
 								NewStreamBlocks.Add(PartByteLocation, ChunkPartConsume, ESearchDir::FromEnd);
@@ -1054,12 +1239,32 @@ namespace BuildPatchServices
 				TMap<FGuid, uint32> NewChunkWindowSizes;
 				TSet<FChunkSearcher::FFileNode*> UpdatedFiles;
 				FChunkSearcher ManifestSearcher(FileManifestList);
-				FManifestBuildStreamerConfig UnknownDataStreamConfig({Configuration.CloudDirectory, NewStreamBlocks});
+				FManifestBuildStreamerConfig UnknownDataStreamConfig({Configuration.DownloadCloudDirectories, NewStreamBlocks});
 				FManifestBuildStreamerDependencies UnknownDataStreamDependencies({ChunkReferenceTrackerFactory.Get(), CloudChunkSourceFactoryB.Get(), StatsCollector.Get(), ManifestB.Get()});
 				TUniquePtr<IManifestBuildStreamer> UnknownDataStream(FBuildStreamerFactory::Create(MoveTemp(UnknownDataStreamConfig), MoveTemp(UnknownDataStreamDependencies)));
-				TUniquePtr<IChunkDataSerialization> ChunkDataSerializationWriter(FChunkDataSerializationFactory::Create(FileSystem.Get(), ManifestB->ManifestMeta.FeatureLevel));
-				FParallelChunkWriterConfig ChunkWriterConfig = FParallelChunkWriterConfig({5, 5, 50, 8, Configuration.CloudDirectory, ManifestB->ManifestMeta.FeatureLevel});
+				TUniquePtr<IChunkDataSerialization> ChunkDataSerializationWriter(FChunkDataSerializationFactory::Create(FileSystem.Get(), Crypto.Get(), { ManifestB->ManifestMeta.FeatureLevel, Configuration.EncryptionSecrets, ManifestB->GetEncryptionSecretId() /* UE5 MERGE TODO : , CloudChunkSourceStat.Get()*/ }));
+				FParallelChunkWriterConfig ChunkWriterConfig = FParallelChunkWriterConfig({5, 5, 50, 8, Configuration.bResaveExistingChunks, Configuration.OutputCloudDirectory, ManifestB->ManifestMeta.FeatureLevel});
 				TUniquePtr<IParallelChunkWriter> ChunkWriter(FParallelChunkWriterFactory::Create(ChunkWriterConfig, FileSystem.Get(), ChunkDataSerializationWriter.Get(), StatsCollector.Get()));
+				// Register chunk writer events and pump from main thread.
+				FTSTicker::FDelegateHandle ChunkWriterPumpHandle;
+				AsyncHelpers::ExecuteOnGameThread<void>([&]() { ChunkWriterPumpHandle = FTSTicker::GetCoreTicker().AddTicker(FTickerDelegate::CreateLambda([&](float) { ChunkWriter->PumpEvents(); return true; })); }).Wait();
+				if (OnChunkFileWritten().IsBound())
+				{
+					ChunkWriter->OnChunkFileWritten().AddLambda([this](const FString& FullFilePath, const FMD5Hash& MD5Hash, const FSHAHash& SHA1Hash)
+					{
+						OnChunkFileWritten().Broadcast(FullFilePath, MD5Hash);
+					});
+				}
+				ChunkWriter->OnChunkFileWriteFailed().AddLambda([this, ChunkWriterPtr = ChunkWriter.Get()](const FString& FullFilePath) {
+					ChunkWriterPtr->Abort();
+				});
+				auto WhenChunkWriterFailed = [&]() {
+					bShouldRun = false;
+					bSuccess = false;
+					UE_LOG(LogChunkDeltaOptimiser, Error, TEXT("Chunk writer aborted"));
+					return FinalStatLogs;
+				};
+
 				StreamBuffer.SetNumUninitialized(0, AllowShrinking);
 				for (const TTuple<FBlockStructure, FChunkPart>& NewChunk : NewChunks)
 				{
@@ -1090,9 +1295,66 @@ namespace BuildPatchServices
 					const FSHAHash NewChunkSha = DeltaOptimiseHelpers::GetShaForDataSet(StreamBuffer.GetData(), StreamBuffer.Num());
 
 					// Save it out.
-					ChunkWriter->AddChunkData(StreamBuffer, NewChunkPart.Guid, NewChunkHash, NewChunkSha);
+					if (!ChunkWriter->AddChunkData(StreamBuffer, NewChunkPart.Guid, NewChunkHash, NewChunkSha))
+					{
+						return WhenChunkWriterFailed();
+					}
+				}
+
+				// Iterate the output file manifest and find how much each file changed, and what the sources were.
+				// Track where each chunk came from so we can report source information for matches.
+				struct FGuidSourceInfo
+				{
+					FString File;
+					uint64 FileOffset, Size;
+					uint32 ReadOffset;
+				};
+				TMap<FGuid, TArray<FGuidSourceInfo>> SourceInfos;
+
+				if (Configuration.SourceDetailsLogFilename.Len())
+				{
+					TArray<FString> SourceFiles;
+					ManifestA->GetFileList(SourceFiles);
+					
+					TMap<uint64, uint32> ChunkSizes;
+					for (FString& SourceFile : SourceFiles)
+					{
+						const FFileManifest* SourceManifest = ManifestA->GetFileManifest(SourceFile);
+						uint64 SourceFileOffset = 0;
+						TSet<FGuid> ChunksInFile;
+
+						for (const FChunkPart& ChunkPart : SourceManifest->ChunkParts)
+						{
+							ChunksInFile.Add(ChunkPart.Guid);
+
+							FGuidSourceInfo NewInfo;
+							NewInfo.File = SourceFile;
+							NewInfo.FileOffset = SourceFileOffset;
+							NewInfo.Size = ChunkPart.Size;
+							SourceFileOffset += ChunkPart.Size;
+							NewInfo.ReadOffset = ChunkPart.Offset;
+
+							SourceInfos.FindOrAdd(ChunkPart.Guid).Add(MoveTemp(NewInfo));
+
+							const BuildPatchServices::FChunkInfo* ChunkInfo = ManifestA->GetChunkInfo(ChunkPart.Guid);
+							ChunkSizes.FindOrAdd(ChunkInfo->DataSizeUncompressed, 0)++;
+						}
+
+						SourceDetailsOutputLines.Add(FString::Printf(TEXT("OldFileChunks: %s"), *SourceFile));
+						for (FGuid& G : ChunksInFile)
+						{
+							SourceDetailsOutputLines.Add(FString::Printf(TEXT("    OFC: %s"), *LexToString(G)));
+						}
+					}
+
+					SourceDetailsOutputLines.Add(TEXT("Old Build Chunk Size Counts"));
+					for (TPair<uint64, uint32> P : ChunkSizes)
+					{
+						SourceDetailsOutputLines.Add(FString::Printf(TEXT("    %s: %s"), *FText::AsNumber(P.Key).ToString(), *FText::AsNumber(P.Value).ToString()));
+					}
 				}
 
+
 				// We also need to potentially upgrade chunks from ManifestA into ManifestB feature level.
 				if (bUsingDifferentChunkSubdir)
 				{
@@ -1123,7 +1385,10 @@ namespace BuildPatchServices
 						checkf(UpgradeChunkDataAccess != nullptr, TEXT("Failed to download chunk from source %s."), *FBuildPatchUtils::GetDataFilename(*ManifestA.Get(), UpgradeChunk));
 						FScopeLockedChunkData LockedChunkData(UpgradeChunkDataAccess);
 						TArray<uint8> ChunkDataArray(LockedChunkData.GetData(), LockedChunkData.GetHeader()->DataSizeUncompressed);
-						ChunkWriter->AddChunkData(MoveTemp(ChunkDataArray), UpgradeChunk, UpgradeChunkInfo->Hash, UpgradeChunkInfo->ShaHash);
+						if (!ChunkWriter->AddChunkData(MoveTemp(ChunkDataArray), UpgradeChunk, UpgradeChunkInfo->Hash, UpgradeChunkInfo->ShaHash))
+						{
+							return WhenChunkWriterFailed();
+						}
 						UpgradeChunkReferenceTracker->PopReference(UpgradeChunk);
 					}
 				}
@@ -1135,7 +1400,10 @@ namespace BuildPatchServices
 				StreamBuffer.SetNumUninitialized(PaddingChunk::ChunkSize);
 				FMemory::Memset(StreamBuffer.GetData(), PaddingChunkId.D, PaddingChunk::ChunkSize);
 				FSHA1::HashBuffer(StreamBuffer.GetData(), PaddingChunk::ChunkSize, PaddingChunkSha.Hash);
-				ChunkWriter->AddChunkData(StreamBuffer, PaddingChunkId, FRollingHash::GetHashForDataSet(StreamBuffer.GetData(), PaddingChunk::ChunkSize), PaddingChunkSha);
+				if (!ChunkWriter->AddChunkData(StreamBuffer, PaddingChunkId, FRollingHash::GetHashForDataSet(StreamBuffer.GetData(), PaddingChunk::ChunkSize), PaddingChunkSha))
+				{
+					return WhenChunkWriterFailed();
+				}
 				for (uint32 LoopIdx = 1; LoopIdx <= 255; ++LoopIdx)
 				{
 					const uint8 Byte = LoopIdx & 0xFF;
@@ -1143,7 +1411,27 @@ namespace BuildPatchServices
 					NewChunkWindowSizes.Add(PaddingChunkId, PaddingChunk::ChunkSize);
 					FMemory::Memset(StreamBuffer.GetData(), PaddingChunkId.D, PaddingChunk::ChunkSize);
 					FSHA1::HashBuffer(StreamBuffer.GetData(), PaddingChunk::ChunkSize, PaddingChunkSha.Hash);
-					ChunkWriter->AddChunkData(StreamBuffer, PaddingChunkId, FRollingHash::GetHashForDataSet(StreamBuffer.GetData(), PaddingChunk::ChunkSize), PaddingChunkSha);
+					if (!ChunkWriter->AddChunkData(StreamBuffer, PaddingChunkId, FRollingHash::GetHashForDataSet(StreamBuffer.GetData(), PaddingChunk::ChunkSize), PaddingChunkSha))
+					{
+						return WhenChunkWriterFailed();
+					}
+				}
+
+				// Wait for chunk writer to be completed.
+				while (ChunkWriter->GetInFlightChunkCount() > 0 && !ChunkWriter->HasAborted())
+				{
+					FPlatformProcess::Sleep(0.01f);
+				}
+
+				if (ChunkWriter->HasAborted())
+				{
+					return WhenChunkWriterFailed();
+				}
+
+				// Use main thread to remove the pump.
+				if (ChunkWriterPumpHandle.IsValid())
+				{
+					AsyncHelpers::ExecuteOnGameThread<void>([&]() { FTSTicker::GetCoreTicker().RemoveTicker(ChunkWriterPumpHandle); }).Wait();
 				}
 
 				// Complete chunk writer.
@@ -1166,8 +1454,71 @@ namespace BuildPatchServices
 				TSet<FGuid> AddedChunkInfos;
 				for (const FFileManifest& FileManifest : DeltaManifest.FileManifestList.FileList)
 				{
+					TMap<FString, uint64> MatchedBytesFromFile;
+					TSet<FGuid> GuidsRequiredFromInstall;
+
+					SourceDetailsOutputLines.Add(FileManifest.Filename);
+
+					uint64 CurrentFileOffset = 0;
 					for (const FChunkPart& ChunkPart : FileManifest.ChunkParts)
 					{
+						if (Configuration.SourceDetailsLogFilename.Len())
+						{
+							// Here we find out if this chunk part is from the old build. If it is, we try and find out 
+							const TArray<FGuidSourceInfo>* SourceInfoPtr = SourceInfos.Find(ChunkPart.Guid);
+						
+							if (!SourceInfoPtr)
+							{
+								MatchedBytesFromFile.FindOrAdd(TEXT("new"), 0) += ChunkPart.Size;
+							}
+							else
+							{
+								// Here on, it's matched data off the install.
+
+								GuidsRequiredFromInstall.Add(ChunkPart.Guid);
+								const TArray<FGuidSourceInfo>& SourceInfo = *SourceInfoPtr;
+
+								if (SourceInfo.Num() == 1)
+								{
+									// The match pulls from a single source file.
+									MatchedBytesFromFile.FindOrAdd(SourceInfo[0].File, 0) += ChunkPart.Size;
+									
+									FString LogLine = FString::Printf(TEXT("    %s: [%s size=%d] came from %s at %s+%d read offset %s"), *LexToString(ChunkPart.Guid), *FText::AsNumber(CurrentFileOffset).ToString(),
+										ChunkPart.Size, *SourceInfo[0].File, *FText::AsNumber(SourceInfo[0].FileOffset).ToString(), ChunkPart.Offset, *FText::AsNumber(SourceInfo[0].ReadOffset).ToString());
+									SourceDetailsOutputLines.Add(MoveTemp(LogLine));
+
+									// <good one> UE_LOG(LogTemp, Warning, TEXT("%s: [%llu size=%d] came from %s at %llu+%d, read offset %d"), *LexToString(ChunkPart.Guid), CurrentOffset, ChunkPart.Size, *SourceInfo[0].File, SourceInfo[0].FileOffset, ChunkPart.Offset, SourceInfo[0].ReadOffset);
+								}
+								else
+								{
+									// The match pulls from multiple files. This is rare enough on large files we don't emit detailed debug info.
+									MatchedBytesFromFile.FindOrAdd(TEXT("mul"), 0) += ChunkPart.Size;
+
+									// We know we need to read ChunkPart.Offset,ChunkPart.Size from the chunk, so see which source provides that range.
+									// This is commented out if we need it.
+
+									//FBlockRange Range = FBlockRange::FromFirstAndSize(ChunkPart.Offset, ChunkPart.Size);
+									//for (const FGuidSourceInfo& Source : SourceInfo)
+									//{
+									//	FBlockRange SourceRange = FBlockRange::FromFirstAndSize(Source.ReadOffset, Source.Size);
+									//	if (SourceRange.Overlaps(Range))
+									//	{
+									//		// This source provides some of the data.
+									//		FBlockRange Intersection = SourceRange.FromIntersection(SourceRange, Range);
+									//
+									//		//UE_LOG(LogTemp, Warning, TEXT("%s MUL: [%llu size=%d] came from %s at %llu+%d (%llu) read offset %d"), *LexToString(ChunkPart.Guid), CurrentOffset, ChunkPart.Size, *Source.File, Source.FileOffset, Intersection.GetFirst(), Intersection.GetSize(), Source.ReadOffset);
+									//	}
+									//}
+									//for (const FGuidSourceInfo& Source : SourceInfo)
+									//{
+									//	UE_LOG(LogTemp, Warning, TEXT("%s MUL: [%llu size=%d] came from %s at %llu+%d (%llu) read offset %d"), *LexToString(ChunkPart.Guid), CurrentOffset, ChunkPart.Size, *Source.File, Source.FileOffset, ChunkPart.Offset, Source.Size, Source.ReadOffset);
+									//}
+								}
+							}
+						} // end if emitting detailed source info
+
+						CurrentFileOffset += ChunkPart.Size;
+
 						bool bWasAlreadyInSet = false;
 						AddedChunkInfos.Add(ChunkPart.Guid, &bWasAlreadyInSet);
 						if (!bWasAlreadyInSet)
@@ -1191,20 +1542,77 @@ namespace BuildPatchServices
 									NewChunkInfo.Hash = ChunkWriterSummaries.ChunkOutputHashes[ChunkPart.Guid];
 									NewChunkInfo.ShaHash = ChunkWriterSummaries.ChunkOutputShas[ChunkPart.Guid];
 									NewChunkInfo.GroupNumber = FCrc::MemCrc32(&ChunkPart.Guid, sizeof(FGuid)) % 100;
-									NewChunkInfo.WindowSize = NewChunkWindowSizes[ChunkPart.Guid];
-									NewChunkInfo.FileSize = ChunkWriterSummaries.ChunkOutputSizes[ChunkPart.Guid];
+									NewChunkInfo.DataSizeUncompressed = NewChunkWindowSizes[ChunkPart.Guid];
+									NewChunkInfo.FileSize = ChunkWriterSummaries.ChunkOutputFileSizes[ChunkPart.Guid];
+									if (const int64* DataSizeCompressed = ChunkWriterSummaries.ChunkOutputCompressedSizes.Find(ChunkPart.Guid))
+									{
+										NewChunkInfo.DataSizeCompressed = *DataSizeCompressed;
+									}
+									if (const FAESAuthTag* AESAuthTag = ChunkWriterSummaries.ChunkOutputAuthTags.Find(ChunkPart.Guid))
+									{
+										NewChunkInfo.AESAuthTag = *AESAuthTag;
+										NewChunkInfo.EncryptionSecretId = ManifestB->GetEncryptionSecretId();
+									}
 								}
 							}
 						}
+					} // end for each chunk park in the file.
+
+					if (Configuration.SourceDetailsLogFilename.Len())
+					{
+						SourceDetailsOutputLines.Add(TEXT("    Matches Bytes From File:"));
+						for (TPair<FString, uint64>& Sizes : MatchedBytesFromFile)
+						{
+							// Note: This is a little bit weird because if the chunk that matches is a full container chunk (i.e. the size of the base build chunk,
+							// not the small optimized scan window), it'll match the entire chunk and end up claiming that it's matching from a random file when it's actually
+							// a patch. There's some sort of accumulation phase I don't have full understanding of. The easiest way to see this in the log is the file source
+							// information a) won't make sense and b) will be exactly a multiple of the base chunk size.
+							SourceDetailsOutputLines.Add(FString::Printf(TEXT("        %s: %s (%.2f)"), *Sizes.Key, *FText::AsNumber(Sizes.Value).ToString(), Sizes.Value * 100.0 / CurrentFileOffset));
+							
+							//UE_LOG(LogTemp, Warning, TEXT("    %s: %llu (%.2f)"), *Sizes.Key, Sizes.Value, Sizes.Value * 100.0 / CurrentFileOffset);
+						}
+						SourceDetailsOutputLines.Add(TEXT("        Guids Requires From Install Source:"));
+						for (FGuid& Guid : GuidsRequiredFromInstall)
+						{
+							SourceDetailsOutputLines.Add(FString::Printf(TEXT("        IC: %s"), *LexToString(Guid)));
+							//UE_LOG(LogTemp, Warning, TEXT("    <%s>"), *LexToString(Guid));
+						}
 					}
-				}
+				} // end each file in the new manifest
+
 				DeltaManifest.InitLookups();
+				if (ManifestB->GetEncryptionSecretId().IsValid())
+				{
+					DeltaManifest.EncryptData(ManifestB->GetEncryptionSecretId(), Configuration.EncryptionSecrets[ManifestB->GetEncryptionSecretId()]);
+				}
 				// For save format, we'll pick the newest of two input manifests, or EFeatureLevel::FirstOptimisedDelta if manifests are older.
 				EFeatureLevel DeltaOutputFormat = FMath::Max3(EFeatureLevel::FirstOptimisedDelta, ManifestA->GetFeatureLevel(), ManifestB->GetFeatureLevel());
+				FMD5Hash MD5Hash;
+				FSHAHash SHAHash;
 				const FString TmpOutputDeltaFilename = OutputDeltaFilename + TEXT("tmp");
-				DeltaManifest.SaveToFile(TmpOutputDeltaFilename, DeltaOutputFormat);
-				FileSystem->MoveFile(*OutputDeltaFilename, *TmpOutputDeltaFilename);
-				FinalStatLogs.Add(FString::Printf(TEXT("Saved new optimised delta file %s"), *OutputDeltaFilename));
+				if (DeltaManifest.SaveToFile(TmpOutputDeltaFilename, DeltaOutputFormat, &SHAHash, &MD5Hash) && FileSystem->MoveFile(*OutputDeltaFilename, *TmpOutputDeltaFilename))
+				{
+					FinalStatLogs.Add(FString::Printf(TEXT("Saved new optimised delta file %s"), *OutputDeltaFilename));
+					// Run the notification for saving delta file.
+					BroadcastDelta(OutputDeltaFilename, MD5Hash, SHAHash);
+				}
+				else
+				{
+					UE_LOG(LogChunkDeltaOptimiser, Error, TEXT("Optimised delta completed successfully but could not be saved %s."), *OutputDeltaFilename);
+					bSuccess = false;
+				}
+
+				if (Configuration.SourceDetailsLogFilename.Len())
+				{
+					if (FFileHelper::SaveStringArrayToFile(SourceDetailsOutputLines, *Configuration.SourceDetailsLogFilename))
+					{
+						FinalStatLogs.Add(FString::Printf(TEXT("Saved source details log %s"), *Configuration.SourceDetailsLogFilename));
+					}
+					else
+					{
+						FinalStatLogs.Add(FString::Printf(TEXT("Failed to save source details log %s"), *Configuration.SourceDetailsLogFilename));
+					}
+				}
 			}
 
 			if (bSuccess)
@@ -1247,11 +1655,18 @@ namespace BuildPatchServices
 					}
 				}
 				int64 DeltaFileSize = 0;
-				if (bRunProcess && !FileSystem->GetFileSize(*OutputDeltaFilename, DeltaFileSize))
+				if (!bDeltaPreviouslyCompleted)
 				{
-					UE_LOG(LogChunkDeltaOptimiser, Error, TEXT("Could not save output to %s"), *OutputDeltaFilename);
-					bSuccess = false;
-					DeltaFileSize = 0;
+					if (bRunProcess && !FileSystem->GetFileSize(*OutputDeltaFilename, DeltaFileSize))
+					{
+						UE_LOG(LogChunkDeltaOptimiser, Error, TEXT("Could not save output to %s"), *OutputDeltaFilename);
+						bSuccess = false;
+						DeltaFileSize = 0;
+					}
+				}
+				else
+				{
+					DeltaFileSize = LoadedDeltaSize;
 				}
 				ProcessTimer.Stop();
 
@@ -1299,27 +1714,39 @@ namespace BuildPatchServices
 		return FinalStatLogs;
 	}
 
-	void FChunkDeltaOptimiser::HandleDownloadComplete(int32 RequestId, const FDownloadRef& Download)
+	void FChunkDeltaOptimiser::BroadcastDelta(const FString& FullFilePath, const FMD5Hash& MD5Hash, const FSHAHash& SHA1Hash)
+	{
+		AsyncHelpers::ExecuteOnGameThread<void>([&]() { OnDeltaFileWritten().Broadcast(FullFilePath, MD5Hash, SHA1Hash); }).Wait();
+	}
+
+	TFuture<FChunkDeltaOptimiser::FFetchResult> FChunkDeltaOptimiser::FetchManifest(const FString& ManifestUri)
+	{
+		TSharedRef<TPromise<FFetchResult>, ESPMode::ThreadSafe> Promise(new TPromise<FFetchResult>());
+		FDownloadCompleteDelegate CompleteDelegate = FDownloadCompleteDelegate::CreateRaw(this, &FChunkDeltaOptimiser::HandleDownloadComplete, Promise);
+		DownloadService->RequestFile(ManifestUri, CompleteDelegate, DownloadProgressDelegate);
+		return Promise->GetFuture();
+	}
+
+	void FChunkDeltaOptimiser::HandleDownloadComplete(int32 RequestId, const FDownloadRef& Download, TSharedRef<TPromise<FFetchResult>, ESPMode::ThreadSafe> Promise)
 	{
-		TPromise<FBuildPatchAppManifestPtr>* RelevantPromisePtr = RequestId == RequestIdManifestA ? &PromiseManifestA : RequestId == RequestIdManifestB ? &PromiseManifestB : nullptr;
-		if (RelevantPromisePtr != nullptr)
+		if (Download->ResponseSuccessful())
 		{
-			if (Download->ResponseSuccessful())
+			Async(EAsyncExecution::ThreadPool, [Download, Promise]()
 			{
-				Async(EAsyncExecution::ThreadPool, [Download, RelevantPromisePtr]()
+				FBuildPatchAppManifestPtr Manifest = MakeShareable(new FBuildPatchAppManifest());
+				if (Manifest->DeserializeFromData(Download->GetData()))
 				{
-					FBuildPatchAppManifestPtr Manifest = MakeShareable(new FBuildPatchAppManifest());
-					if (!Manifest->DeserializeFromData(Download->GetData()))
-					{
-						Manifest.Reset();
-					}
-					RelevantPromisePtr->SetValue(Manifest);
-				});
-			}
-			else
-			{
-				RelevantPromisePtr->SetValue(FBuildPatchAppManifestPtr());
-			}
+					Promise->SetValue(FFetchResult{ Manifest, Download->GetData().Num(), EFetchResult::Success });
+				}
+				else
+				{
+					Promise->SetValue(FFetchResult{ nullptr, Download->GetData().Num(), EFetchResult::CouldNotDeserialize });
+				}
+			});
+		}
+		else
+		{
+			Promise->SetValue(FFetchResult{ nullptr, INDEX_NONE, EFetchResult::DoesNotExist });
 		}
 	}
 
@@ -1345,8 +1772,8 @@ namespace BuildPatchServices
 		return DesiredBytes;
 	}
 
-	IChunkDeltaOptimiser* FChunkDeltaOptimiserFactory::Create(const FChunkDeltaOptimiserConfiguration& Configuration)
+	IChunkDeltaOptimiser* FChunkDeltaOptimiserFactory::Create(FChunkDeltaOptimiserConfiguration Configuration)
 	{
-		return new FChunkDeltaOptimiser(Configuration);
+		return new FChunkDeltaOptimiser(MoveTemp(Configuration));
 	}
 }
diff --git a/Engine/Source/Runtime/Online/BuildPatchServices/Private/Generation/ChunkDeltaOptimiser.h b/Engine/Source/Runtime/Online/BuildPatchServices/Private/Generation/ChunkDeltaOptimiser.h
--- a/Engine/Source/Runtime/Online/BuildPatchServices/Private/Generation/ChunkDeltaOptimiser.h
+++ b/Engine/Source/Runtime/Online/BuildPatchServices/Private/Generation/ChunkDeltaOptimiser.h
@@ -2,20 +2,15 @@
 #pragma once
 
 #include "CoreMinimal.h"
+
+#include "Interfaces/IChunkDeltaOptimiser.h"
 #include "BuildPatchSettings.h"
 
 namespace BuildPatchServices
 {
-	class IChunkDeltaOptimiser
-	{
-	public:
-		virtual ~IChunkDeltaOptimiser() = default;
-		virtual bool Run() = 0;
-	};
-
 	class FChunkDeltaOptimiserFactory
 	{
 	public:
-		static IChunkDeltaOptimiser* Create(const FChunkDeltaOptimiserConfiguration& Configuration);
+		static IChunkDeltaOptimiser* Create(FChunkDeltaOptimiserConfiguration Configuration);
 	};
 }
diff --git a/Engine/Source/Runtime/Online/BuildPatchServices/Private/Generation/ChunkWriter.cpp b/Engine/Source/Runtime/Online/BuildPatchServices/Private/Generation/ChunkWriter.cpp
--- a/Engine/Source/Runtime/Online/BuildPatchServices/Private/Generation/ChunkWriter.cpp
+++ b/Engine/Source/Runtime/Online/BuildPatchServices/Private/Generation/ChunkWriter.cpp
@@ -2,20 +2,26 @@
 
 #include "Generation/ChunkWriter.h"
 
-#include "Logging/LogMacros.h"
+#include "Async/Async.h"
 #include "Containers/Queue.h"
+#include "Containers/Ticker.h"
 #include "HAL/ThreadSafeBool.h"
+#include "Logging/LogMacros.h"
+#include "Misc/Base64.h"
 #include "Misc/Paths.h"
 #include "Misc/OutputDeviceRedirector.h"
-#include "Async/Async.h"
+#include "Misc/SecureHash.h"
 
+#include "Common/Crypto.h"
 #include "Common/FileSystem.h"
 #include "Common/StatsCollector.h"
 #include "Core/AsyncHelpers.h"
 #include "Data/ChunkData.h"
 #include "BuildPatchUtil.h"
+#include "BuildPatchServicesModule.h"
 
-DECLARE_LOG_CATEGORY_CLASS(LogChunkWriter, Log, All);
+DECLARE_LOG_CATEGORY_EXTERN(LogChunkWriter, Log, All);
+DEFINE_LOG_CATEGORY(LogChunkWriter);
 
 namespace BuildPatchServices
 {
@@ -23,6 +29,9 @@ namespace BuildPatchServices
 	typedef TTuple <FGuid, int64> FChunkOutputSize;
 	typedef TTuple <FGuid, uint64> FChunkOutputHash;
 	typedef TTuple <FGuid, FSHAHash> FChunkOutputSha;
+	typedef TTuple <FGuid, FGuid> FChunkOutputSecretId;
+	typedef TTuple <FGuid, FAESAuthTag> FChunkOutputAuthTag;
+	typedef TTuple <FString, FMD5Hash, FSHAHash> FChunkFileCompleteInfo;
 
 	class FWriterChunkDataAccess
 		: public IChunkDataAccess
@@ -65,6 +74,7 @@ namespace BuildPatchServices
 			, StatsCollector(InStatsCollector)
 			, bMoreDataIsExpected(true)
 			, bShouldAbort(false)
+			, InFlightChunks(0)
 		{
 			StatUncompressesData = 0;
 			StatSerlialiseTime = StatsCollector->CreateStat(TEXT("Chunk Writer: Serialize Time"), EStatFormat::Timer);
@@ -72,11 +82,31 @@ namespace BuildPatchServices
 			StatDataWritten = StatsCollector->CreateStat(TEXT("Chunk Writer: Data Size Written"), EStatFormat::DataSize);
 			StatCompressionRatio = StatsCollector->CreateStat(TEXT("Chunk Writer: Compression Ratio"), EStatFormat::Percentage);
 			StatDataWriteSpeed = StatsCollector->CreateStat(TEXT("Chunk Writer: Data Write Speed"), EStatFormat::DataSpeed);
-			FileSystem->MakeDirectory(*Config.ChunkDirectory);
-			if (!FileSystem->DirectoryExists(*Config.ChunkDirectory))
+
+			// In case of network interruption it is needed to recheck again the cloud directory existance
+			bool bMakeDirSuccess = false;
+			int32 RetryCount = Config.OperationRetryCount;
+			while (!bMakeDirSuccess && RetryCount-- >= 0)
+			{
+				if (FileSystem->MakeDirectory(*Config.ChunkDirectory))
+				{
+					bMakeDirSuccess = true;
+					continue;
+				}
+
+				if (RetryCount >= 0)
+				{
+					UE_LOG(LogChunkWriter, Verbose, TEXT("Retry to create cloud directory (%s), attempts left (%d)."), *Config.ChunkDirectory, RetryCount + 1);
+					FPlatformProcess::Sleep(Config.OperationRetryTime);
+				}
+			}
+
+			if (!bMakeDirSuccess)
 			{
-				UE_LOG(LogChunkWriter, Fatal, TEXT("Could not create cloud directory (%s)."), *Config.ChunkDirectory);
+				UE_LOG(LogChunkWriter, Error, TEXT("Could not create cloud directory (%s)."), *Config.ChunkDirectory);
+				check(bMakeDirSuccess); // Abort the program because cloud directory was not created
 			}
+
 			for (int32 ThreadIdx = 0; ThreadIdx < Config.NumberOfThreads; ++ThreadIdx)
 			{
 				WriterThreads.Add(Async(EAsyncExecution::Thread, [this]()
@@ -88,7 +118,7 @@ namespace BuildPatchServices
 
 		~FParallelChunkWriter()
 		{
-			bShouldAbort = true;
+			Abort();
 			for (const TFuture<void>& Thread : WriterThreads)
 			{
 				Thread.Wait();
@@ -97,19 +127,47 @@ namespace BuildPatchServices
 		}
 
 		// IParallelChunkWriter interface begin.
-		virtual void AddChunkData(TArray<uint8> ChunkData, const FGuid& ChunkGuid, const uint64& ChunkHash, const FSHAHash& ChunkSha) override
+		virtual bool AddChunkData(TArray<uint8> ChunkData, const FGuid& ChunkGuid, const uint64& ChunkHash, const FSHAHash& ChunkSha) override
 		{
-			DebugCheckSingleProducer();
-			// Check for violations of feature level.
-			check(Config.FeatureLevel >= EFeatureLevel::VariableSizeChunksWithoutWindowSizeChunkInfo || ChunkData.Num() == LegacyFixedChunkWindow);
-			// Wait for queue space.
-			while (ChunkDataJobQueueCount.GetValue() >= Config.MaxQueueSize)
+			if (!bShouldAbort)
 			{
-				FPlatformProcess::Sleep(0);
+				DebugCheckSingleProducer();
+				// Check for violations of feature level.
+				check(Config.FeatureLevel >= EFeatureLevel::VariableSizeChunksWithoutWindowSizeChunkInfo || ChunkData.Num() == LegacyFixedChunkWindow);
+				InFlightChunks.Increment();
+				// Wait for queue space.
+				while (ChunkDataJobQueueCount.GetValue() >= Config.MaxQueueSize)
+				{
+					if (bShouldAbort)
+					{
+						InFlightChunks.Decrement();
+						return false;
+					}
+					FPlatformProcess::Sleep(0);
+				}
+				ChunkDataJobQueueCount.Increment();
+				ChunkDataJobQueue.Enqueue(FChunkDataJob(MoveTemp(ChunkData), ChunkGuid, ChunkHash, ChunkSha));
+				return true;
 			}
-			ChunkDataJobQueueCount.Increment();
-			ChunkDataJobQueue.Enqueue(FChunkDataJob(MoveTemp(ChunkData), ChunkGuid, ChunkHash, ChunkSha));
+			return false;
+		}
+
+		virtual int32 GetInFlightChunkCount() override
+		{
+			return InFlightChunks.GetValue();
+		}
+
+		virtual void Abort() override
+		{
+			bShouldAbort = true;
+			bMoreDataIsExpected = false;
 		}
+
+		virtual bool HasAborted() const override
+		{
+			return bShouldAbort;
+		}
+
 		virtual FParallelChunkWriterSummaries OnProcessComplete() override
 		{
 			bMoreDataIsExpected = false;
@@ -120,9 +178,13 @@ namespace BuildPatchServices
 			WriterThreads.Empty();
 			ParallelChunkWriterSummaries.FeatureLevel = Config.FeatureLevel;
 			FChunkOutputSize ChunkOutputSize;
-			while (ChunkOutputSizeQueue.Dequeue(ChunkOutputSize))
+			while (ChunkOutputCompressedSizeQueue.Dequeue(ChunkOutputSize))
 			{
-				ParallelChunkWriterSummaries.ChunkOutputSizes.Add(ChunkOutputSize.Get<0>(), ChunkOutputSize.Get<1>());
+				ParallelChunkWriterSummaries.ChunkOutputCompressedSizes.Add(ChunkOutputSize.Get<0>(), ChunkOutputSize.Get<1>());
+			}
+			while (ChunkOutputFileSizeQueue.Dequeue(ChunkOutputSize))
+			{
+				ParallelChunkWriterSummaries.ChunkOutputFileSizes.Add(ChunkOutputSize.Get<0>(), ChunkOutputSize.Get<1>());
 			}
 			FChunkOutputHash ChunkOutputHash;
 			while (ChunkOutputHashQueue.Dequeue(ChunkOutputHash))
@@ -134,13 +196,49 @@ namespace BuildPatchServices
 			{
 				ParallelChunkWriterSummaries.ChunkOutputShas.Add(ChunkOutputSha.Get<0>(), ChunkOutputSha.Get<1>());
 			}
+			FChunkOutputSecretId ChunkOutputSecretId;
+			while (ChunkOutputSecretIdQueue.Dequeue(ChunkOutputSecretId))
+			{
+				ParallelChunkWriterSummaries.ChunkOutputSecretIds.Add(ChunkOutputSecretId.Get<0>(), ChunkOutputSecretId.Get<1>());
+			}
+			FChunkOutputAuthTag ChunkOutputAuthTag;
+			while (ChunkOutputAuthTagQueue.Dequeue(ChunkOutputAuthTag))
+			{
+				ParallelChunkWriterSummaries.ChunkOutputAuthTags.Add(ChunkOutputAuthTag.Get<0>(), ChunkOutputAuthTag.Get<1>());
+			}
 			return ParallelChunkWriterSummaries;
 		}
+
+		virtual FOnChunkFileWritten& OnChunkFileWritten() override
+		{
+			return ChunkFileWrittenEvent;
+		}
+
+		virtual FOnChunkFileWriteFailed& OnChunkFileWriteFailed() override
+		{
+			return ChunkFileWriteFailedEvent;
+		}
+
+		virtual void PumpEvents() override
+		{
+			FChunkFileCompleteInfo CompletionInfo;
+			while (ChunkCompleteQueue.Dequeue(CompletionInfo))
+			{
+				OnChunkFileWritten().Broadcast(CompletionInfo.Get<0>(), CompletionInfo.Get<1>(), CompletionInfo.Get<2>());
+				InFlightChunks.Decrement();
+			}
+			FString FailedFilename;
+			while (ChunkFailedQueue.Dequeue(FailedFilename))
+			{
+				OnChunkFileWriteFailed().Broadcast(FailedFilename);
+			}
+		}
 		// IParallelChunkWriter interface end.
 
 	private:
 		void WriterThread()
 		{
+			TArray<uint8> ChunkFileBytes;
 			uint64 ChunkWriteTime;
 			bool bReceivedJob;
 			FChunkDataJob ChunkDataJob;
@@ -160,44 +258,65 @@ namespace BuildPatchServices
 					ChunkHeader.Guid = ChunkGuid;
 					ChunkHeader.DataSizeCompressed = ChunkData.Num();
 					ChunkHeader.DataSizeUncompressed = ChunkData.Num();
-					ChunkHeader.StoredAs = EChunkStorageFlags::None;
+					ChunkHeader.StoredAs = EChunkStorageFlags::RawData;
 					ChunkHeader.HashType = EChunkHashFlags::RollingPoly64 | EChunkHashFlags::Sha1;
 					ChunkHeader.RollingHash = ChunkHash;
 					ChunkHeader.SHAHash = ChunkSha;
-					TUniquePtr<FWriterChunkDataAccess> ChunkDataAccess(new FWriterChunkDataAccess(ChunkData, ChunkHeader));
-					const FString NewChunkFilename = Config.ChunkDirectory / FBuildPatchUtils::GetChunkNewFilename(Config.FeatureLevel, ChunkGuid, ChunkHash);
-					int64 ExistingChunkFilesize = INDEX_NONE;
-					const bool bFileExists = FileSystem->GetFileSize(*NewChunkFilename, ExistingChunkFilesize);
-					if (!bFileExists)
+					const TUniquePtr<FWriterChunkDataAccess> ChunkDataAccess(new FWriterChunkDataAccess(ChunkData, ChunkHeader));
+					// We start with a memory writer, the only reason this would fail is fatal, so we don't keep retrying the serialisation.
+					// We also need information from the serialisation class. This get put into ChunkHeader, and is then used in chunk filename generation.
+					ChunkFileBytes.Reset();
+					FMemoryWriter ChunkBytesWriter(ChunkFileBytes);
+					FStatsCollector::AccumulateTimeBegin(ChunkWriteTime);
+					const EChunkSaveResult ChunkSaveResult = ChunkDataSerialization->SaveToArchive(ChunkBytesWriter, ChunkDataAccess.Get());
+					FStatsCollector::AccumulateTimeEnd(StatSerlialiseTime, ChunkWriteTime);
+					const int64 ChunkFileSize = ChunkFileBytes.Num();
+					bool bChunkSuccess = ChunkBytesWriter.Close() && ChunkSaveResult == EChunkSaveResult::Success;
+					if (!bChunkSuccess)
 					{
-						bool bSaveSuccess = false;
-						int32 RetryCount = Config.SaveRetryCount;
-						while (!bShouldAbort && !bSaveSuccess && RetryCount-- >= 0)
+						UE_LOG(LogChunkWriter, Log, TEXT("Could not serialise chunk to memory [%d %s]."), ChunkBytesWriter.IsError(), ToString(ChunkSaveResult));
+					}
+					const FString NewChunkFilename = Config.ChunkDirectory / FBuildPatchUtils::GetChunkNewFilename(Config.FeatureLevel, ChunkHeader);
+					const bool bSaveToFile = bChunkSuccess && (Config.bResaveExistingChunks || !FileSystem->FileExists(*NewChunkFilename));
+					if (bSaveToFile)
+					{
+						bChunkSuccess = false;
+						int32 RetryCount = Config.OperationRetryCount;
+						while (!bShouldAbort && !bChunkSuccess && RetryCount-- >= 0)
 						{
 							FStatsCollector::AccumulateTimeBegin(ChunkWriteTime);
 							FileSystem->MakeDirectory(*FPaths::GetPath(NewChunkFilename));
-							TUniquePtr<FArchive> ChunkFileOut = FileSystem->CreateFileWriter(*NewChunkFilename);
-							if (!ChunkFileOut.IsValid())
+							bChunkSuccess = FileSystem->SaveArrayToFile(*NewChunkFilename, ChunkFileBytes) && FileSystem->FileExists(*NewChunkFilename);
+							FStatsCollector::AccumulateTimeEnd(StatSerlialiseTime, ChunkWriteTime);
+							if (!bChunkSuccess)
 							{
-								UE_LOG(LogChunkWriter, Log, TEXT("Could not create chunk (%s)."), *NewChunkFilename);
-								FPlatformProcess::Sleep(Config.SaveRetryTime);
+								UE_LOG(LogChunkWriter, Log, TEXT("Could not create chunk (%s). %s"), *NewChunkFilename, RetryCount > 0 ? TEXT("Retrying...") : TEXT("Fatal!"));
+								FPlatformProcess::Sleep(Config.OperationRetryTime);
 								continue;
 							}
-							EChunkSaveResult ChunkSaveResult = ChunkDataSerialization->SaveToArchive(*ChunkFileOut.Get(), ChunkDataAccess.Get());
-							FStatsCollector::AccumulateTimeEnd(StatSerlialiseTime, ChunkWriteTime);
-							const int64 ChunkFileSize = ChunkFileOut->TotalSize();
-							const bool bArchiveSuccess = ChunkFileOut->Close();
-							const bool bNewFileExists = FileSystem->FileExists(*NewChunkFilename);
-							if (!bArchiveSuccess || !bNewFileExists || ChunkSaveResult != EChunkSaveResult::Success)
+							ChunkOutputCompressedSizeQueue.Enqueue(FChunkOutputSize(ChunkGuid, ChunkHeader.DataSizeCompressed));
+							ChunkOutputFileSizeQueue.Enqueue(FChunkOutputSize(ChunkGuid, ChunkFileSize));
+							if (EnumHasAnyFlags(ChunkHeader.StoredAs, EChunkStorageFlags::Encrypted))
 							{
-								UE_LOG(LogChunkWriter, Log, TEXT("Could not save chunk [%d %d %s] (%s)."), bArchiveSuccess, bNewFileExists, ToString(ChunkSaveResult), *NewChunkFilename);
-								FPlatformProcess::Sleep(Config.SaveRetryTime);
-								continue;
+								ChunkOutputSecretIdQueue.Enqueue(FChunkOutputSecretId(ChunkGuid, ChunkHeader.EncryptionSecretId));
+								ChunkOutputAuthTagQueue.Enqueue(FChunkOutputAuthTag(ChunkGuid, ChunkHeader.AESAuthTag));
 							}
-							bSaveSuccess = true;
-							ChunkOutputSizeQueue.Enqueue(FChunkOutputSize(ChunkGuid, ChunkFileSize));
 							ChunkOutputHashQueue.Enqueue(FChunkOutputHash(ChunkGuid, ChunkHash));
 							ChunkOutputShaQueue.Enqueue(FChunkOutputSha(ChunkGuid, ChunkSha));
+							if (OnChunkFileWritten().IsBound())
+							{
+								FMD5 Md5Gen;
+								FMD5Hash MD5Hash;
+								FSHAHash SHA1Hash;
+								Md5Gen.Update(ChunkFileBytes.GetData(), ChunkFileBytes.Num());
+								MD5Hash.Set(Md5Gen);
+								FSHA1::HashBuffer(ChunkFileBytes.GetData(), ChunkFileBytes.Num(), SHA1Hash.Hash);
+								ChunkCompleteQueue.Enqueue(FChunkFileCompleteInfo(NewChunkFilename, MD5Hash, SHA1Hash));
+							}
+							else
+							{
+								InFlightChunks.Decrement();
+							}
 							FStatsCollector::Accumulate(StatChunksSaved, 1);
 							FStatsCollector::Accumulate(StatDataWritten, ChunkFileSize);
 							FStatsCollector::Accumulate(&StatUncompressesData, ChunkHeader.HeaderSize + ChunkHeader.DataSizeUncompressed);
@@ -213,16 +332,42 @@ namespace BuildPatchServices
 								FStatsCollector::Set(StatDataWriteSpeed, *StatDataWritten / SerlialiseTime);
 							}
 						}
-						if (!bSaveSuccess)
+					}
+					else
+					{
+						TUniquePtr<FArchive> ChunkFile = FileSystem->CreateFileReader(*NewChunkFilename);
+						bChunkSuccess = ChunkFile.IsValid();
+						if (bChunkSuccess)
+						{
+							*ChunkFile << ChunkHeader;
+							ChunkOutputCompressedSizeQueue.Enqueue(FChunkOutputSize(ChunkGuid, ChunkHeader.DataSizeCompressed));
+							ChunkOutputFileSizeQueue.Enqueue(FChunkOutputSize(ChunkGuid, ChunkFile->TotalSize()));
+							ChunkOutputHashQueue.Enqueue(FChunkOutputHash(ChunkGuid, ChunkHeader.RollingHash));
+							ChunkOutputShaQueue.Enqueue(FChunkOutputSha(ChunkGuid, ChunkHeader.SHAHash));
+							if (EnumHasAnyFlags(ChunkHeader.StoredAs, EChunkStorageFlags::Encrypted))
+							{
+								ChunkOutputSecretIdQueue.Enqueue(FChunkOutputSecretId(ChunkGuid, ChunkHeader.EncryptionSecretId));
+								ChunkOutputAuthTagQueue.Enqueue(FChunkOutputAuthTag(ChunkGuid, ChunkHeader.AESAuthTag));
+							}
+							InFlightChunks.Decrement();
+							UE_LOG(LogChunkWriter, Log, TEXT("Skipped save of existing chunk (%s)."), *NewChunkFilename);
+						}
+						else
 						{
-							UE_LOG(LogChunkWriter, Fatal, TEXT("Chunk save failure (%s)."), *NewChunkFilename);
+							UE_LOG(LogChunkWriter, Log, TEXT("Could not read existing chunk (%s)."), *NewChunkFilename);
 						}
 					}
-					else
+					if (!bChunkSuccess)
 					{
-						ChunkOutputSizeQueue.Enqueue(FChunkOutputSize(ChunkGuid, ExistingChunkFilesize));
-						ChunkOutputHashQueue.Enqueue(FChunkOutputHash(ChunkGuid, ChunkHash));
-						ChunkOutputShaQueue.Enqueue(FChunkOutputSha(ChunkGuid, ChunkSha));
+						InFlightChunks.Decrement();
+						if (OnChunkFileWriteFailed().IsBound())
+						{
+							ChunkFailedQueue.Enqueue(NewChunkFilename);
+						}
+						else
+						{
+							UE_LOG(LogChunkWriter, Fatal, TEXT("Failed to serialise chunk (%s)."), *NewChunkFilename);
+						}
 					}
 				}
 				else
@@ -263,16 +408,25 @@ namespace BuildPatchServices
 		TArray<TFuture<void>> WriterThreads;
 		FThreadSafeBool bMoreDataIsExpected;
 		FThreadSafeBool bShouldAbort;
+		FThreadSafeCounter InFlightChunks;
 
 		FCriticalSection ChunkDataJobQueueConsumerCS;
 		TQueue<FChunkDataJob, EQueueMode::Spsc> ChunkDataJobQueue;
 		FThreadSafeInt32 ChunkDataJobQueueCount;
 
-		TQueue<FChunkOutputSize, EQueueMode::Mpsc> ChunkOutputSizeQueue;
+		TQueue<FChunkOutputSize, EQueueMode::Mpsc> ChunkOutputCompressedSizeQueue;
+		TQueue<FChunkOutputSize, EQueueMode::Mpsc> ChunkOutputFileSizeQueue;
 		TQueue<FChunkOutputHash, EQueueMode::Mpsc> ChunkOutputHashQueue;
 		TQueue<FChunkOutputSha, EQueueMode::Mpsc> ChunkOutputShaQueue;
+		TQueue<FChunkOutputSecretId, EQueueMode::Mpsc> ChunkOutputSecretIdQueue;
+		TQueue<FChunkOutputAuthTag, EQueueMode::Mpsc> ChunkOutputAuthTagQueue;
+		TQueue<FChunkFileCompleteInfo, EQueueMode::Mpsc> ChunkCompleteQueue;
+		TQueue<FString, EQueueMode::Mpsc> ChunkFailedQueue;
 		FParallelChunkWriterSummaries ParallelChunkWriterSummaries;
 
+		IParallelChunkWriter::FOnChunkFileWritten ChunkFileWrittenEvent;
+		IParallelChunkWriter::FOnChunkFileWriteFailed ChunkFileWriteFailedEvent;
+
 #if !UE_BUILD_SHIPPING
 		TOptional<uint32> ProducerThreadId;
 #endif
diff --git a/Engine/Source/Runtime/Online/BuildPatchServices/Private/Generation/ChunkWriter.h b/Engine/Source/Runtime/Online/BuildPatchServices/Private/Generation/ChunkWriter.h
--- a/Engine/Source/Runtime/Online/BuildPatchServices/Private/Generation/ChunkWriter.h
+++ b/Engine/Source/Runtime/Online/BuildPatchServices/Private/Generation/ChunkWriter.h
@@ -4,7 +4,9 @@
 #include "CoreMinimal.h"
 #include "Misc/Guid.h"
 #include "Misc/SecureHash.h"
+
 #include "BuildPatchManifest.h"
+#include "Common/Crypto.h"
 
 namespace BuildPatchServices
 {
@@ -14,10 +16,11 @@ namespace BuildPatchServices
 
 	struct FParallelChunkWriterConfig
 	{
-		int32 SaveRetryCount;
-		float SaveRetryTime;
+		int32 OperationRetryCount;
+		float OperationRetryTime;
 		int32 MaxQueueSize;
 		int32 NumberOfThreads;
+		bool bResaveExistingChunks;
 		FString ChunkDirectory;
 		EFeatureLevel FeatureLevel;
 	};
@@ -25,9 +28,12 @@ namespace BuildPatchServices
 	struct FParallelChunkWriterSummaries
 	{
 		EFeatureLevel FeatureLevel;
-		TMap<FGuid, int64> ChunkOutputSizes;
+		TMap<FGuid, int64> ChunkOutputCompressedSizes;
+		TMap<FGuid, int64> ChunkOutputFileSizes;
 		TMap<FGuid, uint64> ChunkOutputHashes;
 		TMap<FGuid, FSHAHash> ChunkOutputShas;
+		TMap<FGuid, FGuid> ChunkOutputSecretIds;
+		TMap<FGuid, FAESAuthTag> ChunkOutputAuthTags;
 	};
 
 	class IParallelChunkWriter
@@ -36,14 +42,69 @@ namespace BuildPatchServices
 		virtual ~IParallelChunkWriter() = default;
 
 		/**
-		 * ...
+		 * Declares an event type exposed on this class for when chunks have been saved to the cloud directory.
+		 * The event offers two types of hash of the filedata saved out.
+		 * @param FullFilePath   The full file path on disk for the file saved.
+		 * @param MD5Hash        The MD5 checksum for the file data.
+		 * @param SHA1Hash       The SHA1 checksum for the file data. This is not the ChunkSha, it is the serialized file SHA1.
+		 */
+		DECLARE_TS_MULTICAST_DELEGATE_ThreeParams(FOnChunkFileWritten, const FString& /*FullFilePath*/, const FMD5Hash& /*MD5Hash*/, const FSHAHash& /*SHA1Hash*/);
+
+		/**
+		 * Declares an event type exposed on this class for when chunks have failed to be saved to the cloud directory.
+		 * The event offers an error message of happened error.
+		 * @param FullFilePath   The full file path on disk for for the file that failed.		 
+		 */
+		DECLARE_TS_MULTICAST_DELEGATE_OneParam(FOnChunkFileWriteFailed, const FString& /*FullFilePath*/);
+
+		/**
+		 * @return a delegate that will broadcast when a chunk was written to the cloud directory.
+		 */
+		virtual FOnChunkFileWritten& OnChunkFileWritten() = 0;
+
+		/**
+		 * @return a delegate that will broadcast when a chunk failed to be written to the cloud directory.
+		 */
+		virtual FOnChunkFileWriteFailed& OnChunkFileWriteFailed() = 0;
+
+		/**
+		 * Must be called periodically to flush events, from the desired receiving thread.
+		 * If this is not called the OnChunkFileWritten() will never broadcast, and GetInFlightChunkCount() can remain high.
 		 */
-		virtual void AddChunkData(TArray<uint8> ChunkData, const FGuid& ChunkGuid, const uint64& ChunkHash, const FSHAHash& ChunkSha) = 0;
+		virtual void PumpEvents() = 0;
 
 		/**
-		 * ...
+		 * Queues up serialization for a chunk file.
+		 * @param ChunkData      The raw chunk data from the build.
+		 * @param ChunkGuid      The identifier for the chunk.
+		 * @param ChunkHash      The rolling hash for the chunk.
+		 * @param ChunkSha       The SHA1 hash for the chunk.
+		 * @return true if queued up successfully.
+		 */
+		virtual bool AddChunkData(TArray<uint8> ChunkData, const FGuid& ChunkGuid, const uint64& ChunkHash, const FSHAHash& ChunkSha) = 0;
+
+		/**
+		 * @return the number of chunk writes still outstanding.
+		 */
+		virtual int32 GetInFlightChunkCount() = 0;
+
+		/**
+		 * Gets the summary data for the process. The function will not return until all threads have completed.
+		 * @return the summary data upon completion of the process.
 		 */
 		virtual FParallelChunkWriterSummaries OnProcessComplete() = 0;
+
+		/**
+		 * Aborts current operations and no longer accepts new chunks.
+		 */
+		virtual void Abort() = 0;
+
+		/**
+		 * Gets if the writer has aborted
+		 * @return True if it has aborted
+		 */
+		virtual bool HasAborted() const = 0;
+
 	};
 
 	class FParallelChunkWriterFactory
