p4 integrate //UE5/Dev-Distro-5.5/Engine/Source/Programs/BuildPatchTool/... //UE5/Main/Engine/Source/Programs/BuildPatchTool/...

diff --git a/Engine/Source/Runtime/Online/BuildPatchServices/Private/Generation/CloudEnumeration.cpp b/Engine/Source/Runtime/Online/BuildPatchServices/Private/Generation/CloudEnumeration.cpp
--- a/Engine/Source/Runtime/Online/BuildPatchServices/Private/Generation/CloudEnumeration.cpp
+++ b/Engine/Source/Runtime/Online/BuildPatchServices/Private/Generation/CloudEnumeration.cpp
@@ -1,36 +1,64 @@
 // Copyright Epic Games, Inc. All Rights Reserved.
 
 #include "Generation/CloudEnumeration.h"
-#include "HAL/FileManager.h"
-#include "Misc/ScopeLock.h"
-#include "Async/Future.h"
+
 #include "Async/Async.h"
-#include "Core/BlockStructure.h"
+#include "Async/Future.h"
 #include "BuildPatchManifest.h"
 #include "BuildPatchUtil.h"
+#include "Core/BlockStructure.h"
+#include "HAL/FileManager.h"
+#include "Misc/ConfigCacheIni.h"
+#include "Misc/FileHelper.h"
+#include "Misc/ScopeLock.h"
 
 DECLARE_LOG_CATEGORY_EXTERN(LogCloudEnumeration, Log, All);
 DEFINE_LOG_CATEGORY(LogCloudEnumeration);
 
 namespace BuildPatchServices
 {
+	struct FTriggerExitScope
+	{
+		FTriggerExitScope(FEvent* InEvent)
+			: Event(InEvent)
+		{}
+		~FTriggerExitScope()
+		{
+			Event->Trigger();
+		}
+		FEvent* Event;
+	};
+
 	class FCloudEnumeration
 		: public ICloudEnumeration
 	{
 	public:
-		FCloudEnumeration(const FString& CloudDirectory, const FDateTime& ManifestAgeThreshold, const EFeatureLevel& OutputFeatureLevel, FStatsCollector* StatsCollector);
+		/**
+		 * Creates an instance of FCloudEnumeration for managing cloud uploaded manifests and chunks.
+		 *
+		 * @param CloudDirectory The cloud root directory.
+		 * @param ManifestAgeThreshold The threshold age for manifest files; files older than this will be ignored.
+		 * @param OutputFeatureLevel The desired feature level for enumerated output.
+		 * @param StatsCollector A pointer to the statistics collector for recording related data.
+		 */
+		FCloudEnumeration(const FString& CloudDirectory, const FDateTime& ManifestAgeThreshold, const EFeatureLevel& InOutputFeatureLevel, FStatsCollector* StatsCollector);
 		virtual ~FCloudEnumeration();
 
 		virtual bool IsComplete() const override;
 		virtual const TSet<uint32>& GetUniqueWindowSizes() const override;
 		virtual const TMap<uint64, TSet<FGuid>>& GetChunkInventory() const override;
 		virtual const TMap<FGuid, int64>& GetChunkFileSizes() const override;
+		virtual const TMap<FGuid, int64>& GetChunkCompressedSizes() const override;
 		virtual const TMap<FGuid, FSHAHash>& GetChunkShaHashes() const override;
+		virtual const TMap<FGuid, FAESAuthTag>& GetChunkAESAuthTags() const override;
 		virtual const TMap<FGuid, uint32>& GetChunkWindowSizes() const override;
+		virtual const TMap<FGuid, FGuid>& GetChunkSecretIds() const override;
 		virtual bool IsChunkFeatureLevelMatch(const FGuid& ChunkId) const override;
 		virtual const uint64& GetChunkHash(const FGuid& ChunkId) const override;
 		virtual const FSHAHash& GetChunkShaHash(const FGuid& ChunkId) const override;
 		virtual const TMap<FSHAHash, TSet<FGuid>>& GetIdenticalChunks() const override;
+		virtual const TArray<FString>& GetUnsupportedFeatureLevelBuilds() const override;
+
 	private:
 		void EnumerateCloud();
 		TFuture<FBuildPatchAppManifestPtr> AsyncLoadManifest(const FString& ManifestFilename);
@@ -41,17 +69,24 @@ namespace BuildPatchServices
 		const FString CloudDirectory;
 		const FDateTime ManifestAgeThreshold;
 		const TCHAR* FeatureLevelChunkSubdir;
+		FEvent* ManifestLoadedEvent;
 		TMap<uint64, TSet<FGuid>> ChunkInventory;
 		TMap<FGuid, int64> ChunkFileSizes;
+		TMap<FGuid, int64> ChunkCompressedSizes;
 		TMap<FGuid, uint64> ChunkHashes;
 		TMap<FGuid, FSHAHash> ChunkShaHashes;
+		TMap<FGuid, FAESAuthTag> ChunkAESAuthTags;
 		TSet<uint32> UniqueWindowSizes;
 		TMap<FGuid, uint32> ChunkWindowSizes;
 		TMap<uint32, TSet<FGuid>> WindowsSizeChunks;
 		TSet<FGuid> FeatureLevelMatchedChunks;
 		TMap<FSHAHash, TSet<FGuid>> IdenticalChunks;
+		TMap<FGuid, FGuid> ChunkSecretIds;
+		TArray<FString> UnsupportedBuildVersions;
 		FStatsCollector* StatsCollector;
 		TFuture<void> Future;
+		volatile FStatsCollector::FAtomicValue* StatManifestsFound;
+		volatile FStatsCollector::FAtomicValue* StatManifestsLoading;
 		volatile FStatsCollector::FAtomicValue* StatManifestsLoaded;
 		volatile FStatsCollector::FAtomicValue* StatManifestsRejected;
 		volatile FStatsCollector::FAtomicValue* StatChunksEnumerated;
@@ -60,13 +95,16 @@ namespace BuildPatchServices
 		volatile FStatsCollector::FAtomicValue* StatUniqueWindowSizes;
 	};
 
-	FCloudEnumeration::FCloudEnumeration(const FString& InCloudDirectory, const FDateTime& InManifestAgeThreshold, const EFeatureLevel& OutputFeatureLevel, FStatsCollector* InStatsCollector)
+	FCloudEnumeration::FCloudEnumeration(const FString& InCloudDirectory, const FDateTime& InManifestAgeThreshold, const EFeatureLevel& InOutputFeatureLevel, FStatsCollector* InStatsCollector)
 		: CloudDirectory(InCloudDirectory)
 		, ManifestAgeThreshold(InManifestAgeThreshold)
-		, FeatureLevelChunkSubdir(ManifestVersionHelpers::GetChunkSubdir(OutputFeatureLevel))
+		, FeatureLevelChunkSubdir(ManifestVersionHelpers::GetChunkSubdir(InOutputFeatureLevel))
+		, ManifestLoadedEvent(FPlatformProcess::GetSynchEventFromPool())
 		, StatsCollector(InStatsCollector)
 	{
 		// Create statistics.
+		StatManifestsFound = StatsCollector->CreateStat(TEXT("Cloud Enumeration: Manifests Found"), EStatFormat::Value);
+		StatManifestsLoading = StatsCollector->CreateStat(TEXT("Cloud Enumeration: Manifests Loading"), EStatFormat::Value);
 		StatManifestsLoaded = StatsCollector->CreateStat(TEXT("Cloud Enumeration: Manifests Loaded"), EStatFormat::Value);
 		StatManifestsRejected = StatsCollector->CreateStat(TEXT("Cloud Enumeration: Manifests Rejected"), EStatFormat::Value);
 		StatChunksEnumerated = StatsCollector->CreateStat(TEXT("Cloud Enumeration: Chunks Enumerated"), EStatFormat::Value);
@@ -82,6 +120,7 @@ namespace BuildPatchServices
 	FCloudEnumeration::~FCloudEnumeration()
 	{
 		Future.Wait();
+		FPlatformProcess::ReturnSynchEventToPool(ManifestLoadedEvent);
 	}
 
 	bool FCloudEnumeration::IsComplete() const
@@ -107,11 +146,23 @@ namespace BuildPatchServices
 		return ChunkFileSizes;
 	}
 
+	const TMap<FGuid, int64>& FCloudEnumeration::GetChunkCompressedSizes() const
+	{
+		Future.Wait();
+		return ChunkCompressedSizes;
+	}
+
 	const TMap<FGuid, FSHAHash>& FCloudEnumeration::GetChunkShaHashes() const
 	{
 		Future.Wait();
 		return ChunkShaHashes;
 	}
+	
+	const TMap<FGuid, FAESAuthTag>& FCloudEnumeration::GetChunkAESAuthTags() const
+	{
+		Future.Wait();
+		return ChunkAESAuthTags;
+	}
 
 	const TMap<FGuid, uint32>& FCloudEnumeration::GetChunkWindowSizes() const
 	{
@@ -119,6 +170,12 @@ namespace BuildPatchServices
 		return ChunkWindowSizes;
 	}
 
+	const TMap<FGuid, FGuid>& FCloudEnumeration::GetChunkSecretIds() const
+	{
+		Future.Wait();
+		return ChunkSecretIds;
+	}
+
 	bool FCloudEnumeration::IsChunkFeatureLevelMatch(const FGuid& ChunkId) const
 	{
 		Future.Wait();
@@ -143,10 +200,21 @@ namespace BuildPatchServices
 		return IdenticalChunks;
 	}
 
+	const TArray<FString>& FCloudEnumeration::GetUnsupportedFeatureLevelBuilds() const
+	{
+		Future.Wait();
+		return UnsupportedBuildVersions;
+	}
+
 	void FCloudEnumeration::EnumerateCloud()
 	{
-		uint64 EnumerationTimer;
+		// Grab a config for parallel load count.
+		int32 NumParallelManifestLoads = 10;
+		GConfig->GetInt(TEXT("Portal.BuildPatch"), TEXT("CloudParallelManifestLoads"), NumParallelManifestLoads, GEngineIni);
+		NumParallelManifestLoads = FMath::Clamp(NumParallelManifestLoads, 1, 100);
+		UE_LOG(LogCloudEnumeration, Log, TEXT("Loading manifests %d at a time"), NumParallelManifestLoads);
 
+		uint64 EnumerationTimer;
 		IFileManager& FileManager = IFileManager::Get();
 
 		// Find all manifest files
@@ -155,26 +223,41 @@ namespace BuildPatchServices
 		{
 			TArray<FString> AllManifestFilenames;
 			FileManager.FindFiles(AllManifestFilenames, *(CloudDirectory / TEXT("*.manifest")), true, false);
+			FStatsCollector::Set(StatManifestsFound, AllManifestFilenames.Num());
 			FStatsCollector::AccumulateTimeEnd(StatTotalTime, EnumerationTimer);
 			FStatsCollector::AccumulateTimeBegin(EnumerationTimer);
 
-			// Load all manifest files
-			TArray<TFuture<FBuildPatchAppManifestPtr>> AllManifestFutures;
-			for (const FString& ManifestFilename : AllManifestFilenames)
+			// Load and process all manifest files
+			TArray<TFuture<FBuildPatchAppManifestPtr>> ManifestFutures;
+			while (AllManifestFilenames.Num() > 0 || ManifestFutures.Num() > 0)
 			{
-				AllManifestFutures.Add(AsyncLoadManifest(CloudDirectory / ManifestFilename));
-			}
-			// Process all manifest files
-			for (const TFuture<FBuildPatchAppManifestPtr>& ManifestFuture : AllManifestFutures)
-			{
-				// Determine chunks from manifest file
-				FBuildPatchAppManifestPtr BuildManifest = ManifestFuture.Get();
-				if (BuildManifest.IsValid())
+				// Process all completed manifests
+				TArray<TFuture<FBuildPatchAppManifestPtr>> TempManifestFutures = MoveTemp(ManifestFutures);
+				for (TFuture<FBuildPatchAppManifestPtr>& ManifestFuture : TempManifestFutures)
+				{
+					if (ManifestFuture.IsReady())
+					{
+						// Determine chunks from manifest file
+						FBuildPatchAppManifestPtr BuildManifest = ManifestFuture.Get();
+						if (BuildManifest.IsValid())
+						{
+							EnumerateManifestData(BuildManifest.ToSharedRef());
+						}
+						FStatsCollector::AccumulateTimeEnd(StatTotalTime, EnumerationTimer);
+						FStatsCollector::AccumulateTimeBegin(EnumerationTimer);
+					}
+					else
+					{
+						ManifestFutures.Add(MoveTemp(ManifestFuture));
+					}
+				}
+				// Kick off new manifests
+				while (ManifestFutures.Num() < NumParallelManifestLoads && AllManifestFilenames.Num() > 0)
 				{
-					EnumerateManifestData(BuildManifest.ToSharedRef());
+					ManifestFutures.Add(AsyncLoadManifest(CloudDirectory / AllManifestFilenames.Pop()));
 				}
-				FStatsCollector::AccumulateTimeEnd(StatTotalTime, EnumerationTimer);
-				FStatsCollector::AccumulateTimeBegin(EnumerationTimer);
+				// Wait for any manifest to complete (with defensive max 20s).
+				ManifestLoadedEvent->Wait(2000);
 			}
 		}
 		else
@@ -188,14 +271,29 @@ namespace BuildPatchServices
 	{
 		return Async(EAsyncExecution::ThreadPool, [this, ManifestFilename]() -> FBuildPatchAppManifestPtr
 		{
-			if (IFileManager::Get().GetTimeStamp(*ManifestFilename) < ManifestAgeThreshold)
+			FTriggerExitScope TriggerExitScope(ManifestLoadedEvent);
+			FStatsScopedCounter StatsScopedCounter(StatManifestsLoading);
+			const uint64 TimeStampStart = FStatsCollector::GetCycles();
+			const FDateTime TimeStamp = IFileManager::Get().GetTimeStamp(*ManifestFilename);
+			const uint64 TimeStampEnd = FStatsCollector::GetCycles();
+			UE_LOG(LogCloudEnumeration, Log, TEXT("Timestamp %s for manifest %s (in %.2f sec)"), *TimeStamp.ToString(), *ManifestFilename, FStatsCollector::CyclesToSeconds(TimeStampEnd - TimeStampStart));
+			if (TimeStamp < ManifestAgeThreshold)
 			{
 				FStatsCollector::Accumulate(StatManifestsRejected, 1);
 				return nullptr;
 			}
+			TArray<uint8> FileData;
+			const uint64 LoadStart = FStatsCollector::GetCycles();
+			const bool bLoadSuccess = FFileHelper::LoadFileToArray(FileData, *ManifestFilename);
+			const uint64 LoadEnd = FStatsCollector::GetCycles();
+			UE_LOG(LogCloudEnumeration, Log, TEXT("Loaded manifest %s (%s in %.2f sec)"), *ManifestFilename, *FText::AsMemory(FileData.Num(), EMemoryUnitStandard::SI).ToString(), FStatsCollector::CyclesToSeconds(LoadEnd - LoadStart));
 			FBuildPatchAppManifestRef BuildManifest = MakeShareable(new FBuildPatchAppManifest());
-			if (BuildManifest->LoadFromFile(ManifestFilename))
+			const uint64 DeserializeStart = FStatsCollector::GetCycles();
+			if (bLoadSuccess && BuildManifest->DeserializeFromData(FileData))
 			{
+				const uint64 DeserializeEnd = FStatsCollector::GetCycles();
+				UE_LOG(LogCloudEnumeration, Log, TEXT("Deserialized manifest %s (%s in %.2f sec)"), *ManifestFilename, *FText::AsMemory(FileData.Num(), EMemoryUnitStandard::SI).ToString(), FStatsCollector::CyclesToSeconds(DeserializeEnd - DeserializeStart));
+				UE_LOG(LogCloudEnumeration, Log, TEXT("Manifest %s details: BuildVersion: %s, FeatureLevel: %s"), *ManifestFilename, *BuildManifest->GetVersionString(), LexToString(BuildManifest->GetFeatureLevel()));
 				FStatsCollector::Accumulate(StatManifestsLoaded, 1);
 				return BuildManifest;
 			}
@@ -213,6 +311,10 @@ namespace BuildPatchServices
 		typedef TTuple<TSet<uint32>, TMap<FGuid, uint32>> FWindowSizes;
 		// Check if this chunk will already beling to our matching feature level sub dir - ptr==ptr test is fine, no need to string compare.
 		const bool bMatchingChunkSubdir = FeatureLevelChunkSubdir == ManifestVersionHelpers::GetChunkSubdir(Manifest->GetFeatureLevel());
+		if (Manifest->GetFeatureLevel() > EFeatureLevel::Latest)
+		{
+			UnsupportedBuildVersions.Add(Manifest->GetVersionString());
+		}
 		TFuture<FWindowSizes> CalculateChunkWindowSizesFuture = Async(EAsyncExecution::TaskGraph, [&](){ return CalculateChunkWindowSizes(Manifest); });
 		TArray<FGuid> DataList;
 		Manifest->GetDataList(DataList);
@@ -230,10 +332,19 @@ namespace BuildPatchServices
 					{
 						// Added new chunk, fill out all the various tracking data.
 						ChunkFileSizes.Add(DataGuid, Manifest->GetDataSize(DataGuid));
+						if (ChunkInfo->DataSizeCompressed > 0)
+						{
+							ChunkCompressedSizes.Add(DataGuid, ChunkInfo->DataSizeCompressed);
+						}
 						ChunkHashes.Add(DataGuid, ChunkInfo->Hash);
 						FMemory::Memcpy(ChunkShaHashes.FindOrAdd(DataGuid).Hash, ChunkInfo->ShaHash.Hash, FSHA1::DigestSize);
-						UniqueWindowSizes.Add(ChunkInfo->WindowSize);
-						ChunkWindowSizes.Add(DataGuid, ChunkInfo->WindowSize);
+						if (ChunkInfo->EncryptionSecretId.IsValid())
+						{
+							FMemory::Memcpy(ChunkAESAuthTags.FindOrAdd(DataGuid).AuthTag, ChunkInfo->AESAuthTag.AuthTag, AES256_GCM_AuthTagSizeInBytes);
+							ChunkSecretIds.Add(DataGuid, ChunkInfo->EncryptionSecretId);
+						}
+						UniqueWindowSizes.Add(ChunkInfo->DataSizeUncompressed);
+						ChunkWindowSizes.Add(DataGuid, ChunkInfo->DataSizeUncompressed);
 						IdenticalChunks.FindOrAdd(ChunkInfo->ShaHash).Add(DataGuid);
 						FStatsCollector::Accumulate(StatChunksEnumerated, 1);
 					}
diff --git a/Engine/Source/Runtime/Online/BuildPatchServices/Private/Generation/CloudEnumeration.h b/Engine/Source/Runtime/Online/BuildPatchServices/Private/Generation/CloudEnumeration.h
--- a/Engine/Source/Runtime/Online/BuildPatchServices/Private/Generation/CloudEnumeration.h
+++ b/Engine/Source/Runtime/Online/BuildPatchServices/Private/Generation/CloudEnumeration.h
@@ -4,6 +4,7 @@
 #include "CoreMinimal.h"
 #include "Misc/Guid.h"
 #include "Misc/SecureHash.h"
+#include "Common/Crypto.h"
 #include "Common/StatsCollector.h"
 #include "BuildPatchFeatureLevel.h"
 
@@ -17,17 +18,32 @@ namespace BuildPatchServices
 		virtual const TSet<uint32>& GetUniqueWindowSizes() const = 0;
 		virtual const TMap<uint64, TSet<FGuid>>& GetChunkInventory() const = 0;
 		virtual const TMap<FGuid, int64>& GetChunkFileSizes() const = 0;
+		virtual const TMap<FGuid, int64>& GetChunkCompressedSizes() const = 0;
 		virtual const TMap<FGuid, FSHAHash>& GetChunkShaHashes() const = 0;
+		virtual const TMap<FGuid, FAESAuthTag>& GetChunkAESAuthTags() const = 0;
 		virtual const TMap<FGuid, uint32>& GetChunkWindowSizes() const = 0;
+		virtual const TMap<FGuid, FGuid>& GetChunkSecretIds() const = 0;
 		virtual bool IsChunkFeatureLevelMatch(const FGuid& ChunkId) const = 0;
 		virtual const uint64& GetChunkHash(const FGuid& ChunkId) const = 0;
 		virtual const FSHAHash& GetChunkShaHash(const FGuid& ChunkId) const = 0;
 		virtual const TMap<FSHAHash, TSet<FGuid>>& GetIdenticalChunks() const = 0;
+		/* Returns the versions of build which generated with Feature Level greater than EFeatureLevel::Last. */
+		virtual const TArray<FString>& GetUnsupportedFeatureLevelBuilds() const = 0;
 	};
 
 	class FCloudEnumerationFactory
 	{
 	public:
+		/**
+		 * Creates an instance of ICloudEnumeration for managing cloud uploaded chunks.
+		 *
+		 * @param CloudDirectory The cloud root directory.
+		 * @param ManifestAgeThreshold The threshold age for manifest files; files older than this will be ignored.
+		 * @param OutputFeatureLevel The desired feature level for enumerated output.
+		 * @param StatsCollector A pointer to the statistics collector for recording related data.
+		 *
+		 * @return A pointer to the created ICloudEnumeration instance.
+		 */
 		static ICloudEnumeration* Create(const FString& CloudDirectory, const FDateTime& ManifestAgeThreshold, const EFeatureLevel& OutputFeatureLevel, FStatsCollector* StatsCollector);
 	};
 }
diff --git a/Engine/Source/Runtime/Online/BuildPatchServices/Private/Generation/DataScanner.cpp b/Engine/Source/Runtime/Online/BuildPatchServices/Private/Generation/DataScanner.cpp
--- a/Engine/Source/Runtime/Online/BuildPatchServices/Private/Generation/DataScanner.cpp
+++ b/Engine/Source/Runtime/Online/BuildPatchServices/Private/Generation/DataScanner.cpp
@@ -5,19 +5,62 @@
 #include "HAL/ThreadSafeBool.h"
 #include "Async/Future.h"
 #include "Async/Async.h"
+#include "Misc/CommandLine.h"
+#include "Misc/ConfigCacheIni.h"
 
 #include "Core/ProcessTimer.h"
 #include "Data/ChunkData.h"
 #include "Generation/CloudEnumeration.h"
 #include "BuildPatchHash.h"
 
+DECLARE_LOG_CATEGORY_EXTERN(LogDataScanner, Log, All);
+DEFINE_LOG_CATEGORY(LogDataScanner);
+
 namespace BuildPatchServices
 {
+	bool ReadForceSkipMatches()
+	{
+		FString ForceSkipMatches;
+		if (!FParse::Value(FCommandLine::Get(), TEXT("ForceSkipMatches="), ForceSkipMatches))
+		{
+			GConfig->GetString(TEXT("Portal.BuildPatch"), TEXT("ForceSkipMatches"), ForceSkipMatches, GEngineIni);
+		}
+		return ForceSkipMatches == TEXT("true");
+	}
+
+	bool GetForceSkipMatches()
+	{
+		static bool bForceSkipMatches = ReadForceSkipMatches();
+		return bForceSkipMatches;
+	}
+
+	uint32 ReadMaxCollisionCountPerMB()
+	{
+		uint32 CommandlineValue = TNumericLimits<uint32>::Max();
+		if (!FParse::Value(FCommandLine::Get(), TEXT("MaxCollisionCountPerMB="), CommandlineValue))
+		{
+			int32 ConfigValue = 0;
+			if (GConfig->GetInt(TEXT("Portal.BuildPatch"), TEXT("MaxCollisionCountPerMB"), ConfigValue, GEngineIni))
+			{
+				return static_cast<uint32>(ConfigValue);
+			}
+		}
+		return CommandlineValue;
+	}
+
+	uint32 GetCollisionCountLimit(uint64 DataSize)
+	{
+		static uint64 MaxCollisionCountPerMB = ReadMaxCollisionCountPerMB();
+		const uint64 NumMBsInData = DataSize / 1000000;
+		const uint64 CollisionCountLimit = NumMBsInData * MaxCollisionCountPerMB;
+		return FMath::Min<uint64>(TNumericLimits<uint32>::Max(), CollisionCountLimit);
+	}
+
 	class FDataScanner
 		: public IDataScanner
 	{
 	public:
-		FDataScanner(const TArray<uint32>& ChunkWindowSizes, const TArray<uint8>& Data, const ICloudEnumeration* CloudEnumeration, FStatsCollector* StatsCollector);
+		FDataScanner(const TArray<uint32>& ChunkWindowSizes, const TArray<uint8>& Data, const ICloudEnumeration* CloudEnumeration, FStatsCollector* StatsCollector, uint64 DataOffset);
 		virtual ~FDataScanner();
 
 		virtual bool IsComplete() override;
@@ -40,10 +83,14 @@ namespace BuildPatchServices
 		const TMap<uint64, TSet<FGuid>>& ChunkInventory;
 		const TMap<FGuid, FSHAHash>& ChunkShaHashes;
 		const TMap<FSHAHash, TSet<FGuid>>& IdenticalChunks;
+		const uint64 DataOffset;
+		const uint32 CollisionCountLimit;
+		const bool bForceSkipMatches;
 		FThreadSafeBool bIsComplete;
 		FThreadSafeBool bShouldAbort;
 		TFuture<TArray<FChunkMatch>> FutureResult;
 		FProcessTimer ScanTimer;
+		uint32 CollisionCount;
 
 		// Fast forward tech variables.
 		const int32 MatchHistorySize;
@@ -68,7 +115,7 @@ namespace BuildPatchServices
 		static FThreadSafeCounter NumRunningScanners;
 	};
 
-	FDataScanner::FDataScanner(const TArray<uint32>& InChunkWindowSizes, const TArray<uint8>& InData, const ICloudEnumeration* InCloudEnumeration, FStatsCollector* InStatsCollector)
+	FDataScanner::FDataScanner(const TArray<uint32>& InChunkWindowSizes, const TArray<uint8>& InData, const ICloudEnumeration* InCloudEnumeration, FStatsCollector* InStatsCollector, uint64 InDataOffset)
 		: ChunkWindowSizes(InChunkWindowSizes)
 		, Data(InData)
 		, CloudEnumeration(InCloudEnumeration)
@@ -76,8 +123,12 @@ namespace BuildPatchServices
 		, ChunkInventory(CloudEnumeration->GetChunkInventory())
 		, ChunkShaHashes(CloudEnumeration->GetChunkShaHashes())
 		, IdenticalChunks(CloudEnumeration->GetIdenticalChunks())
+		, DataOffset(InDataOffset)
+		, CollisionCountLimit(GetCollisionCountLimit(Data.Num()))
+		, bForceSkipMatches(GetForceSkipMatches())
 		, bIsComplete(false)
 		, bShouldAbort(false)
+		, CollisionCount(0)
 		, MatchHistorySize(100)
 		, MatchHistoryNextIdx(0)
 		, MatchHistoryNextOffset(0)
@@ -226,7 +277,16 @@ namespace BuildPatchServices
 			}
 			else
 			{
+				UE_LOG(LogDataScanner, VeryVerbose, TEXT("@%llu: Collision at %d with chunkhash %X SHA %s"), DataOffset, DataStart, RollingHash.GetWindowHash(), *OutChunkSha.ToString());
+				++CollisionCount;
 				FStatsCollector::Accumulate(StatHashCollisions, 1);
+				// Check for aborting based on high collisions. This is an emergency feature for rare problems
+				// with unusually high processing times.
+				if (CollisionCount > CollisionCountLimit)
+				{
+					UE_LOG(LogDataScanner, Log, TEXT("@%llu: Aborting scanner with %d collisions!"), DataOffset, CollisionCount);
+					bShouldAbort = true;
+				}
 			}
 		}
 		// Set blank in the history.
@@ -325,7 +385,7 @@ namespace BuildPatchServices
 						TempMatchIdx = InsertMatch(DataScanResult, TempMatchIdx, DataStart, ChunkMatch, WindowSize);
 					}
 					// We can start skipping over the chunk that we matched if we have no overlap potential, i.e. we know this match will not be rejected.
-					const bool bCanSkipData = bFoundChunkMatch && !bChunkOverlap;
+					const bool bCanSkipData = bFoundChunkMatch && (!bChunkOverlap || bForceSkipMatches);
 					if (bCanSkipData)
 					{
 						RollingHash.Clear();
@@ -366,6 +426,8 @@ namespace BuildPatchServices
 		// Count running scanners.
 		NumRunningScanners.Decrement();
 
+		UE_LOG(LogDataScanner, Log, TEXT("@%llu: Scanner completed in %s with %d collisions."), DataOffset, *FPlatformTime::PrettyTime(ScanTimer.GetSeconds()), CollisionCount);
+
 		bIsComplete = true;
 		return DataScanResult;
 	}
@@ -403,8 +465,8 @@ namespace BuildPatchServices
 		FDataScanner::NumRunningScanners.Decrement();
 	}
 
-	IDataScanner* FDataScannerFactory::Create(const TArray<uint32>& ChunkWindowSizes, const TArray<uint8>& Data, const ICloudEnumeration* CloudEnumeration, FStatsCollector* StatsCollector)
+	IDataScanner* FDataScannerFactory::Create(const TArray<uint32>& ChunkWindowSizes, const TArray<uint8>& Data, const ICloudEnumeration* CloudEnumeration, FStatsCollector* StatsCollector, uint64 DataOffset)
 	{
-		return new FDataScanner(ChunkWindowSizes, Data, CloudEnumeration, StatsCollector);
+		return new FDataScanner(ChunkWindowSizes, Data, CloudEnumeration, StatsCollector, DataOffset);
 	}
 }
diff --git a/Engine/Source/Runtime/Online/BuildPatchServices/Private/Generation/DataScanner.h b/Engine/Source/Runtime/Online/BuildPatchServices/Private/Generation/DataScanner.h
--- a/Engine/Source/Runtime/Online/BuildPatchServices/Private/Generation/DataScanner.h
+++ b/Engine/Source/Runtime/Online/BuildPatchServices/Private/Generation/DataScanner.h
@@ -29,7 +29,14 @@ namespace BuildPatchServices
 		uint32 WindowSize;
 	};
 
-	typedef TTuple<FBlockRange, FFilenameId, TSet<FString>, uint64> FScannerFileElement;
+	struct FScannerFileElement
+	{
+		FBlockRange FileScannerRange;
+		FFilenameId ScanFilename;
+		TSet<FString> ScanTagSet;
+		uint64 ScannerFileLocation;
+	};
+	
 	typedef TDoubleLinkedList<FScannerFileElement> FScannerFilesList;
 	typedef FScannerFilesList::TDoubleLinkedListNode FScannerFilesListNode;
 
@@ -58,7 +65,7 @@ namespace BuildPatchServices
 	class FDataScannerFactory
 	{
 	public:
-		static IDataScanner* Create(const TArray<uint32>& ChunkWindowSizes, const TArray<uint8>& Data, const ICloudEnumeration* CloudEnumeration, FStatsCollector* StatsCollector);
+		static IDataScanner* Create(const TArray<uint32>& ChunkWindowSizes, const TArray<uint8>& Data, const ICloudEnumeration* CloudEnumeration, FStatsCollector* StatsCollector, uint64 DataOffset);
 	};
 
 	class FDeltaScannerFactory
diff --git a/Engine/Source/Runtime/Online/BuildPatchServices/Private/Generation/DeltaEnumeration.cpp b/Engine/Source/Runtime/Online/BuildPatchServices/Private/Generation/DeltaEnumeration.cpp
--- a/Engine/Source/Runtime/Online/BuildPatchServices/Private/Generation/DeltaEnumeration.cpp
+++ b/Engine/Source/Runtime/Online/BuildPatchServices/Private/Generation/DeltaEnumeration.cpp
@@ -38,12 +38,13 @@ namespace BuildPatchServices
 		: public IDeltaChunkEnumeration
 	{
 	public:
-		FDeltaChunkEnumeration(IManifestBuildStreamer* Streamer, FStatsCollector* StatsCollector, const FBuildPatchAppManifest& Manifest, const uint32 WindowSize);
+		FDeltaChunkEnumeration(IManifestBuildStreamer* Streamer, FStatsCollector* StatsCollector, const FBuildPatchAppManifest& Manifest, const uint32 WindowSize, const TMap<FString, FString>& InAliasedTags);
 		~FDeltaChunkEnumeration();
 
 		// FDeltaChunkEnumeration interface begin.
 		virtual void Run() override;
 		virtual bool IsComplete() const override;
+		virtual bool IsSuccess() const override;
 		virtual const TMap<uint64, TSet<FDeltaChunkId>>& GetChunkInventory() const override;
 		virtual const TMap<FDeltaChunkId, FShaId>& GetChunkShaHashes() const override;
 		virtual const TMap<FDeltaChunkId, FChunkBuildReference>& GetChunkBuildReferences() const override;
@@ -72,14 +73,16 @@ namespace BuildPatchServices
 		TSet<FString> FilenameTable;
 		TSet<FSHAHash> ShaTable;
 		bool bHasRan;
+		bool bIsSuccess;
 		TMap<uint64, TSet<FDeltaChunkId>> ChunkInventory;
 		TMap<FDeltaChunkId, uint64> ChunkHashes;
 		TMap<FDeltaChunkId, FShaId> ChunkShaHashes;
 		TMap<FDeltaChunkId, FChunkBuildReference> ChunkBuildReferences;
 		TMap<FShaId, TSet<FDeltaChunkId>> IdenticalChunks;
+		TMap<FString, FString> AliasedTags;
 	};
 
-	FDeltaChunkEnumeration::FDeltaChunkEnumeration(IManifestBuildStreamer* InStreamer, FStatsCollector* InStatsCollector, const FBuildPatchAppManifest& InManifest, const uint32 InWindowSize)
+	FDeltaChunkEnumeration::FDeltaChunkEnumeration(IManifestBuildStreamer* InStreamer, FStatsCollector* InStatsCollector, const FBuildPatchAppManifest& InManifest, const uint32 InWindowSize, const TMap<FString, FString>& InAliasedTags)
 		: Streamer(InStreamer)
 		, StatsCollector(InStatsCollector)
 		, ChunkSearcher(InManifest)
@@ -87,19 +90,21 @@ namespace BuildPatchServices
 		, WindowSize(InWindowSize)
 		, UniqueWindowSizes({WindowSize})
 		, bHasRan(false)
+		, bIsSuccess(true)
+		, AliasedTags(InAliasedTags)
 	{
 		// Make sure we enumerate special chunk SHAs.
 		FDeltaChunkId PaddingChunkId = PaddingChunk::MakePaddingGuid(0);
 		FShaId PaddingShaId = MakeShaId(DeltaChunkEnumerationHelpers::GetPaddingChunkSha(0, WindowSize));
 		uint64 PaddingHash = DeltaChunkEnumerationHelpers::GetPaddingChunkHash(0, WindowSize);
 		FChunkBuildReference PaddingReference;
-		PaddingReference.Get<0>().AddDefaulted_GetRef().Guid = PaddingChunkId;
-		PaddingReference.Get<0>()[0].Offset = 0;
-		PaddingReference.Get<0>()[0].Size = WindowSize;
-		PaddingReference.Get<3>() = 0;
+		PaddingReference.ChunkParts.AddDefaulted_GetRef().Guid = PaddingChunkId;
+		PaddingReference.ChunkParts[0].Offset = 0;
+		PaddingReference.ChunkParts[0].Size = WindowSize;
+		PaddingReference.FileReferences.AddDefaulted_GetRef().Location = 0;
 
 		// Add already setup variables for byte 0.
-		ChunkInventory.FindOrAdd(PaddingHash).Add(PaddingChunkId);
+		ChunkInventory.Add(PaddingHash);
 		ChunkHashes.Add(PaddingChunkId, PaddingHash);
 		ChunkShaHashes.Add(PaddingChunkId, PaddingShaId);
 		ChunkBuildReferences.Add(PaddingChunkId, PaddingReference);
@@ -110,8 +115,8 @@ namespace BuildPatchServices
 			PaddingChunkId.D = Byte;
 			PaddingShaId = MakeShaId(DeltaChunkEnumerationHelpers::GetPaddingChunkSha(Byte, WindowSize));
 			PaddingHash = DeltaChunkEnumerationHelpers::GetPaddingChunkHash(Byte, WindowSize);
-			PaddingReference.Get<0>()[0].Guid = PaddingChunkId;
-			ChunkInventory.FindOrAdd(PaddingHash).Add(PaddingChunkId);
+			PaddingReference.ChunkParts[0].Guid = PaddingChunkId;
+			ChunkInventory.Add(PaddingHash);
 			ChunkHashes.Add(PaddingChunkId, PaddingHash);
 			ChunkShaHashes.Add(PaddingChunkId, PaddingShaId);
 			ChunkBuildReferences.Add(PaddingChunkId, PaddingReference);
@@ -127,6 +132,11 @@ namespace BuildPatchServices
 		return bHasRan;
 	}
 
+	bool FDeltaChunkEnumeration::IsSuccess() const
+	{
+		return bHasRan && bIsSuccess;
+	}
+
 	const TMap<uint64, TSet<FDeltaChunkId>>& FDeltaChunkEnumeration::GetChunkInventory() const
 	{
 		return ChunkInventory;
@@ -164,7 +174,14 @@ namespace BuildPatchServices
 			// Calculate the buffer's build structure.
 			FBlockStructure StreamBuildStructure;
 			const uint64 StreamStructureSize = StreamerBlocks.SelectSerialBytes(StreamStartPosition, StreamBuffer.Num(), StreamBuildStructure);
-			check(SizeRead == StreamBufferReadSize || Streamer->IsEndOfData());
+
+			if (!(SizeRead == StreamBufferReadSize || Streamer->IsEndOfData()))
+			{
+				UE_LOG(LogDeltaEnumeration, Error, TEXT("Failed to read from chunk stream."));
+				bIsSuccess = false;
+				break;
+			}
+
 			check(StreamStructureSize == StreamBuffer.Num());
 
 			// Chunk up data.
@@ -280,7 +297,7 @@ namespace BuildPatchServices
 		if (ChunkHash != 0)
 		{
 			FDeltaChunkId ChunkId = FGuid::NewGuid();
-			ChunkInventory.FindOrAdd(ChunkHash).Add(ChunkId);
+			ChunkInventory.Add(ChunkHash);
 			ChunkHashes.Add(ChunkId, ChunkHash);
 			ChunkShaHashes.Add(ChunkId, MakeShaId(GetShaForDataSet(&StreamBuffer[StreamBufferPosition], WindowSize)));
 			FBlockStructure ChunkBuildStructure;
@@ -290,21 +307,24 @@ namespace BuildPatchServices
 			{
 				const FChunkSearcher::FChunkNode& ChunkNode = Chunk->GetValue();
 				FChunkBuildReference& ChunkBuildReference = ChunkBuildReferences.FindOrAdd(ChunkId);
-				TArray<FChunkPart>& ChunkParts = ChunkBuildReference.Get<0>();
-				FFilenameId& FirstFile = ChunkBuildReference.Get<1>();
-				TSet<FString>& FirstFileTagset = ChunkBuildReference.Get<2>();
-				uint64& FirstFileLocation = ChunkBuildReference.Get<3>();
 				const int64 ChunkLeftChop = OverlapRange.GetFirst() - ChunkNode.BuildRange.GetFirst();
 				const int64 ChunkRightChop = ChunkNode.BuildRange.GetLast() - OverlapRange.GetLast();
 				const int64 SizeChop = ChunkLeftChop + ChunkRightChop;
-				if (ChunkParts.Num() == 0)
+
+				// Track where we came from.
+				const FChunkSearcher::FFileNode& FileNode = File->GetValue();
+				FChunkBuildReferenceSourceFile& SourceFile = ChunkBuildReference.FileReferences.AddDefaulted_GetRef();
+				SourceFile.Filename = MakeFilenameId(FileNode.Manifest->Filename);
+				SourceFile.Location = (ChunkNode.BuildRange.GetFirst() - FileNode.BuildRange.GetFirst()) + ChunkLeftChop;
+				// If we have any tag aliases, swap out the tags so that they all match the same.
+				SourceFile.Tagset.Reserve(FileNode.Manifest->InstallTags.Num());
+				for (const FString& InstallTag : FileNode.Manifest->InstallTags)
 				{
-					const FChunkSearcher::FFileNode& FileNode = File->GetValue();
-					FirstFile = MakeFilenameId(FileNode.Manifest->Filename);
-					FirstFileTagset.Append(FileNode.Manifest->InstallTags);
-					FirstFileLocation = (ChunkNode.BuildRange.GetFirst() - FileNode.BuildRange.GetFirst()) + ChunkLeftChop;
+					const FString* AliasedTag = AliasedTags.Find(InstallTag);
+					SourceFile.Tagset.Add(AliasedTag ? *AliasedTag : InstallTag);
 				}
-				ChunkParts.Emplace(ChunkNode.ChunkPart.Guid, ChunkNode.ChunkPart.Offset + ChunkLeftChop, ChunkNode.ChunkPart.Size - SizeChop);
+				
+				ChunkBuildReference.ChunkParts.Emplace(ChunkNode.ChunkPart.Guid, ChunkNode.ChunkPart.Offset + ChunkLeftChop, ChunkNode.ChunkPart.Size - SizeChop);
 			});
 		}
 	}
@@ -316,8 +336,8 @@ namespace BuildPatchServices
 		return SHAHash;
 	}
 
-	IDeltaChunkEnumeration* FDeltaChunkEnumerationFactory::Create(IManifestBuildStreamer* Streamer, FStatsCollector* StatsCollector, const FBuildPatchAppManifest& Manifest, const uint32 WindowSize)
+	IDeltaChunkEnumeration* FDeltaChunkEnumerationFactory::Create(IManifestBuildStreamer* Streamer, FStatsCollector* StatsCollector, const FBuildPatchAppManifest& Manifest, const uint32 WindowSize, const TMap<FString, FString>& AliasedTags)
 	{
-		return new FDeltaChunkEnumeration(Streamer, StatsCollector, Manifest, WindowSize);
+		return new FDeltaChunkEnumeration(Streamer, StatsCollector, Manifest, WindowSize, AliasedTags);
 	}
 }
diff --git a/Engine/Source/Runtime/Online/BuildPatchServices/Private/Generation/DeltaEnumeration.h b/Engine/Source/Runtime/Online/BuildPatchServices/Private/Generation/DeltaEnumeration.h
--- a/Engine/Source/Runtime/Online/BuildPatchServices/Private/Generation/DeltaEnumeration.h
+++ b/Engine/Source/Runtime/Online/BuildPatchServices/Private/Generation/DeltaEnumeration.h
@@ -19,7 +19,19 @@ namespace BuildPatchServices
 	typedef FGuid FDeltaChunkId;
 	typedef FSetElementId FFilenameId;
 	typedef FSetElementId FShaId;
-	typedef TTuple<TArray<FChunkPart>, FFilenameId, TSet<FString>, uint64> FChunkBuildReference;
+
+	struct FChunkBuildReferenceSourceFile
+	{
+		FFilenameId Filename;
+		TSet<FString> Tagset;
+		uint64 Location;
+	};
+
+	struct FChunkBuildReference
+	{
+		TArray<FChunkPart> ChunkParts;
+		TArray<FChunkBuildReferenceSourceFile> FileReferences;
+	};
 
 	class IDeltaChunkEnumeration
 	{
@@ -27,6 +39,7 @@ namespace BuildPatchServices
 		virtual ~IDeltaChunkEnumeration() = default;
 		virtual void Run() = 0;
 		virtual bool IsComplete() const = 0;
+		virtual bool IsSuccess() const = 0;
 		virtual const TMap<uint64, TSet<FDeltaChunkId>>& GetChunkInventory() const = 0;
 		virtual const TMap<FDeltaChunkId, FShaId>& GetChunkShaHashes() const = 0;
 		virtual const TMap<FDeltaChunkId, FChunkBuildReference>& GetChunkBuildReferences() const = 0;
@@ -44,6 +57,6 @@ namespace BuildPatchServices
 	class FDeltaChunkEnumerationFactory
 	{
 	public:
-		static IDeltaChunkEnumeration* Create(IManifestBuildStreamer* Streamer, FStatsCollector* StatsCollector, const FBuildPatchAppManifest& Manifest, const uint32 WindowSize);
+		static IDeltaChunkEnumeration* Create(IManifestBuildStreamer* Streamer, FStatsCollector* StatsCollector, const FBuildPatchAppManifest& Manifest, const uint32 WindowSize, const TMap<FString, FString>& InAliasedTags);
 	};
 }
diff --git a/Engine/Source/Runtime/Online/BuildPatchServices/Private/Generation/DeltaScanner.cpp b/Engine/Source/Runtime/Online/BuildPatchServices/Private/Generation/DeltaScanner.cpp
--- a/Engine/Source/Runtime/Online/BuildPatchServices/Private/Generation/DeltaScanner.cpp
+++ b/Engine/Source/Runtime/Online/BuildPatchServices/Private/Generation/DeltaScanner.cpp
@@ -44,9 +44,16 @@ namespace BuildPatchServices
 			, FileLocOffset(0)
 		{ }
 
+		// This is the match as reported by the hash comparisions. It's offset and size should always exactly match the Range below.
 		const FChunkMatch Match;
+
+		// This is the range in the new version data where we have the match. 
 		const FBlockRange Range;
+
+		// Whether the match came from the same source file by name.
 		const bool bSameFile;
+
+		// This is the distance to the match in the source stream. It's used for selecting closer matches in the event of multiple valid matches.
 		const uint64 FileLocOffset;
 	};
 
@@ -72,7 +79,8 @@ namespace BuildPatchServices
 			}
 			else
 			{
-				// First advance idx if no overlap.
+				// First advance idx if no overlap in the source data. If the new match is distinct, then
+				// we end up at Matches.Num(), which falls through to accepting the match.
 				while (LastAcceptedIdx < Matches.Num() && !Matches[LastAcceptedIdx].Range.Overlaps(NewMatchMeta.Range))
 				{
 					++LastAcceptedIdx;
@@ -322,6 +330,7 @@ namespace BuildPatchServices
 			ChunkShaId = CloudEnumeration->GetShaId(ChunkSha);
 			if (ChunkShaId.IsValidId())
 			{
+				// These are the actual chunks we got from the base build. We might have duplicates. We don't get to know where they came from.
 				ChunkMatches = &IdenticalChunks[ChunkShaId];
 				return;
 			}
@@ -347,30 +356,27 @@ namespace BuildPatchServices
 
 		// Setup current file node.
 		// We safely assume this data has been setup correctly to never fail to find.
-		while (!FilesListNode->GetValue().Get<0>().Overlaps(ScanLocation))
+		while (!FilesListNode->GetValue().FileScannerRange.Overlaps(ScanLocation))
 		{
 			FilesListNode = FilesListNode->GetNextNode();
 		}
 
 		const FChunkBuildReference& MatchBuildReference = ChunkBuildReference[ChunkMatch];
 		const FScannerFileElement& ScanFileInfo = FilesListNode->GetValue();
-		const FFilenameId& MatchFilename = MatchBuildReference.Get<1>();
-		const FFilenameId& ScanFilename = ScanFileInfo.Get<1>();
-		const TSet<FString>& MatchTagSet = MatchBuildReference.Get<2>();
-		const TSet<FString>& ScanTagSet = ScanFileInfo.Get<2>();
-		const FBlockRange& FileScannerRange = ScanFileInfo.Get<0>();
-		const uint64& MatchFileLocation = MatchBuildReference.Get<3>();
-		const uint64& ScannerFileLocation = ScanFileInfo.Get<3>();
-
-		check(ScanLocation.GetFirst() >= FileScannerRange.GetFirst());
-		const uint64 ScanFileLocation = ScannerFileLocation + (ScanLocation.GetFirst() - FileScannerRange.GetFirst());
+
+		check(ScanLocation.GetFirst() >= ScanFileInfo.FileScannerRange.GetFirst());
+		const uint64 ScanFileLocation = ScanFileInfo.ScannerFileLocation + (ScanLocation.GetFirst() - ScanFileInfo.FileScannerRange.GetFirst());
 		bSameFile = false;
-		if (MatchTagSet.Num() == 0 || MatchTagSet.Includes(ScanTagSet))
+
+		// We only accept matches from the same tagset. Note that it matches _any_ tag in the set, so having multiple tags
+		// can run into issues where if the tag sets are not installed, this might think they are and expect an on-disk match
+		// that doesn't exist.
+		if (MatchBuildReference.FileReferences[0].Tagset.Num() == 0 || MatchBuildReference.FileReferences[0].Tagset.Includes(ScanFileInfo.ScanTagSet))
 		{
-			bSameFile = MatchFilename == ScanFilename;
+			bSameFile = MatchBuildReference.FileReferences[0].Filename == ScanFileInfo.ScanFilename;
 			if (bSameFile)
 			{
-				LocationOffset = FMath::Max<uint64>(MatchFileLocation, ScanFileLocation) - FMath::Min<uint64>(MatchFileLocation, ScanFileLocation);
+				LocationOffset = FMath::Max<uint64>(MatchBuildReference.FileReferences[0].Location, ScanFileLocation) - FMath::Min<uint64>(MatchBuildReference.FileReferences[0].Location, ScanFileLocation);
 			}
 			else
 			{
@@ -421,7 +427,7 @@ namespace BuildPatchServices
 				uint32 FastForwardCount = 0;
 				const uint32 DataStart = NextByte - WindowSize;
 				const FBlockRange CurrentRange(FBlockRange::FromFirstAndSize(DataStart, WindowSize));
-				// Check for a chunk match at this offset.
+				// Check for a chunk match at this offset. All matches are identical copies - we could have it in many locations
 				FindChunkDataMatches(RollingHash, ChunkMatches, ChunkSha);
 				const bool bFoundChunkMatch = ChunkMatches != nullptr;
 				bSameFile = false;
@@ -429,6 +435,9 @@ namespace BuildPatchServices
 				{
 					// Process each match for contextual info and add to the selector.
 					bool bAccepted = false;
+
+					// FindChunkDataMatches returns a bunch of matches that are the same data, and we need to pick one. Any of them will provide
+					// the data we want.
 					for (const FDeltaChunkId& ChunkMatch : *ChunkMatches)
 					{
 						if (ProcessMatch(ChunkMatch, CurrentRange, bSameFile, FileLocOffset))
diff --git a/Engine/Source/Runtime/Online/BuildPatchServices/Private/Generation/DirectoryBuildStreamer.cpp b/Engine/Source/Runtime/Online/BuildPatchServices/Private/Generation/DirectoryBuildStreamer.cpp
--- a/Engine/Source/Runtime/Online/BuildPatchServices/Private/Generation/DirectoryBuildStreamer.cpp
+++ b/Engine/Source/Runtime/Online/BuildPatchServices/Private/Generation/DirectoryBuildStreamer.cpp
@@ -14,6 +14,10 @@
 #include <unistd.h>
 #endif
 
+#if WITH_OPENSSL
+#include <openssl/sha.h>
+#endif // WITH_OPENSSL
+
 DECLARE_LOG_CATEGORY_EXTERN(LogBuildStreamer, Log, All);
 DEFINE_LOG_CATEGORY(LogBuildStreamer);
 
@@ -25,10 +29,11 @@ namespace BuildPatchServices
 		FileBufferSize = 1024 * 1024 * 10,      // 10 MiB
 		StreamBufferSize = 1024 * 1024 * 100    // 100 MiB
 	};
+	static_assert(StreamBufferSize > (FileBufferSize * 2), "Stream buffers should be able to take at least two enqueues without a dequeue.");
 
 	static bool IsUnixExecutable(const TCHAR* Filename)
 	{
-#if PLATFORM_MAC
+#if PLATFORM_MAC || PLATFORM_LINUX
 		struct stat FileInfo;
 		if (stat(TCHAR_TO_UTF8(Filename), &FileInfo) == 0)
 		{
@@ -40,9 +45,9 @@ namespace BuildPatchServices
 
 	static FString GetSymlinkTarget(const TCHAR* Filename)
 	{
-#if PLATFORM_MAC
-		ANSICHAR SymlinkTarget[MAC_MAX_PATH] = { 0 };
-		if (readlink(TCHAR_TO_UTF8(Filename), SymlinkTarget, MAC_MAX_PATH) != -1)
+#if PLATFORM_MAC || PLATFORM_LINUX
+		ANSICHAR SymlinkTarget[PLATFORM_MAX_FILEPATH_LENGTH_DEPRECATED] = { 0 };
+		if (readlink(TCHAR_TO_UTF8(Filename), SymlinkTarget, PLATFORM_MAX_FILEPATH_LENGTH_DEPRECATED) != -1)
 		{
 			return UTF8_TO_TCHAR(SymlinkTarget);
 		}
@@ -50,6 +55,250 @@ namespace BuildPatchServices
 		return TEXT("");
 	}
 
+	class FSHA256
+	{
+	public:
+		enum { DigestSize = 32 };
+		FSHA256()
+		{
+			SHA256_Init(&ShaContext);
+		}
+
+		void Update(const uint8* Data, uint64 ByteSize)
+		{
+			SHA256_Update(&ShaContext, Data, ByteSize);
+		}
+
+		void Final()
+		{
+			SHA256_Final(Digest.GetData(), &ShaContext);
+		}
+
+		void GetHash(uint8* OutHash)
+		{
+			FMemory::Memcpy(OutHash, Digest.GetData(), DigestSize);
+		}
+
+	private:
+		// Stub noops for platforms without OPENSSL.
+#if !WITH_OPENSSL
+		struct SHA256_CTX {};
+		void SHA256_Init(SHA256_CTX*) {}
+		void SHA256_Update(SHA256_CTX*, const uint8*, uint64) {}
+		void SHA256_Final(const uint8*, SHA256_CTX*) {}
+#endif // !WITH_OPENSSL
+
+	private:
+		SHA256_CTX ShaContext;
+		TArray<uint8, TFixedAllocator<DigestSize>> Digest;
+	};
+
+	struct FMIMETypeLookup
+	{
+	public:
+		enum { NumTypes = 11 };
+		FMIMETypeLookup()
+		{
+			int32 i = 0;
+			const char* NoWild = "nnnnnnnnnnnn";
+			// gif detection
+			const TCHAR* GifName = TEXT("image/gif");
+			const uint8* GifSig1 = (const uint8*)"\x47\x49\x46\x38\x37\x61";
+			const uint8* GifSig2 = (const uint8*)"\x47\x49\x46\x38\x39\x61";
+			Names[i] = GifName; Signatures[i] = GifSig1; Wildcards[i] = NoWild; SignatureSizes[i] = 6; SignatureOffsets[i++] = 0;
+			Names[i] = GifName; Signatures[i] = GifSig2; Wildcards[i] = NoWild; SignatureSizes[i] = 6; SignatureOffsets[i++] = 0;
+			// png detection
+			const TCHAR* PngName = TEXT("image/png");
+			const uint8* PngSig = (const uint8*)"\x89\x50\x4E\x47\x0D\x0A\x1A\x0A";
+			Names[i] = PngName; Signatures[i] = PngSig; Wildcards[i] = NoWild; SignatureSizes[i] = 8; SignatureOffsets[i++] = 0;
+			// tif detection
+			const TCHAR* TifName = TEXT("image/tiff");
+			const uint8* TifSig1 = (const uint8*)"\x49\x49\x2A\x00";
+			const uint8* TifSig2 = (const uint8*)"\x4D\x4D\x00\x2A";
+			Names[i] = TifName; Signatures[i] = TifSig1; Wildcards[i] = NoWild; SignatureSizes[i] = 4; SignatureOffsets[i++] = 0;
+			Names[i] = TifName; Signatures[i] = TifSig2; Wildcards[i] = NoWild; SignatureSizes[i] = 4; SignatureOffsets[i++] = 0;
+			// bmp detection
+			const TCHAR* BmpName = TEXT("image/bmp");
+			const uint8* BmpSig = (const uint8*)"\x42\x4D";
+			Names[i] = BmpName; Signatures[i] = BmpSig; Wildcards[i] = NoWild; SignatureSizes[i] = 2; SignatureOffsets[i++] = 0;
+			// jpg detection
+			const TCHAR* JpgName = TEXT("image/jpeg");
+			const uint8* JpgSig1 = (const uint8*)"\xFF\xD8\xFF\xDB";
+			const uint8* JpgSig2 = (const uint8*)"\xFF\xD8\xFF\xE0\x00\x10\x4A\x46\x49\x46\x00\x01";
+			const uint8* JpgSig3 = (const uint8*)"\xFF\xD8\xFF\xEE";
+			const uint8* JpgSig4 = (const uint8*)"\xFF\xD8\xFF\xE1??\x45\x78\x69\x66\x00\x00";
+			const char* JpgWild4 = "nnnnyynnnnnn";
+			Names[i] = JpgName; Signatures[i] = JpgSig1; Wildcards[i] = NoWild;   SignatureSizes[i] = 4;  SignatureOffsets[i++] = 0;
+			Names[i] = JpgName; Signatures[i] = JpgSig2; Wildcards[i] = NoWild;   SignatureSizes[i] = 12; SignatureOffsets[i++] = 0;
+			Names[i] = JpgName; Signatures[i] = JpgSig3; Wildcards[i] = NoWild;   SignatureSizes[i] = 4;  SignatureOffsets[i++] = 0;
+			Names[i] = JpgName; Signatures[i] = JpgSig4; Wildcards[i] = JpgWild4; SignatureSizes[i] = 12; SignatureOffsets[i++] = 0;
+			// exe/dll detection
+			const TCHAR* ExeName = TEXT("application/vnd.microsoft.portable-executable");
+			const uint8* ExeSig = (const uint8*)"\x4D\x5A";
+			Names[i] = ExeName; Signatures[i] = ExeSig; Wildcards[i] = NoWild;   SignatureSizes[i] = 2;  SignatureOffsets[i++] = 0;
+
+			checkf(i == NumTypes, TEXT("NumTypes is wrong, should be %d"), i);
+		}
+
+		const TCHAR* Names[NumTypes];
+		const uint8* Signatures[NumTypes];
+		const char* Wildcards[NumTypes];
+		uint8  SignatureSizes[NumTypes];
+		uint8  SignatureOffsets[NumTypes];
+	};
+
+	class FMIMEType
+	{
+	public:
+		FMIMEType()
+			: bCalculated(false)
+		{
+		}
+
+		void Update(const uint8* Data, uint64 ByteSize)
+		{
+			if (!bCalculated)
+			{
+				bCalculated = true;
+				for (int32 MIMEIdx = 0; MIMEIdx < FMIMETypeLookup::NumTypes; ++MIMEIdx)
+				{
+					const int32 SigSize = MIMETypeLookup.SignatureSizes[MIMEIdx];
+					const int32 SigOffset = MIMETypeLookup.SignatureOffsets[MIMEIdx];
+					const int32 SizeNeeded = SigOffset + SigSize;
+					check(SigSize >= 0);
+					if (SizeNeeded <= ByteSize)
+					{
+						bool bMatched = true;
+						for (int32 ByteIdx = 0; ByteIdx < SigSize; ByteIdx++)
+						{
+							if (MIMETypeLookup.Wildcards[MIMEIdx][ByteIdx] == 'n' && MIMETypeLookup.Signatures[MIMEIdx][ByteIdx] != Data[SigOffset + ByteIdx])
+							{
+								bMatched = false;
+								break;
+							}
+						}
+						if (bMatched)
+						{
+							MIMEType = MIMETypeLookup.Names[MIMEIdx];
+							return;
+						}
+					}
+				}
+			}
+		}
+
+		const FString& GetMIMEType()
+		{
+			return MIMEType;
+		}
+
+	private:
+		static const FMIMETypeLookup MIMETypeLookup;
+		FString MIMEType;
+		bool bCalculated;
+	};
+	const FMIMETypeLookup FMIMEType::MIMETypeLookup;
+
+	class FAsyncFileHahser
+	{
+	public:
+		FAsyncFileHahser(const FDirectoryBuildStreamerConfig& InConfig)
+			: Config(InConfig)
+		{
+			DataCopy.SetNumUninitialized(BuildPatchServices::FileBufferSize);
+		}
+
+		void AsyncUpdate(const uint8* Data, uint64 ByteSize)
+		{
+			check(ByteSize <= BuildPatchServices::FileBufferSize);
+			Wait();
+			FMemory::Memcpy(DataCopy.GetData(), Data, ByteSize);
+			HashFutures[0] = Async(EAsyncExecution::ThreadPool, [this, ByteSize]() { FileSHA1Hasher.Update(DataCopy.GetData(), ByteSize); });
+			if (Config.bDetectMimeType)
+			{
+				HashFutures[1] = Async(EAsyncExecution::ThreadPool, [this, ByteSize]() { MIMETypeCalc.Update(DataCopy.GetData(), ByteSize); });
+			}
+			if (Config.bCalculateMD5)
+			{
+				HashFutures[2] = Async(EAsyncExecution::ThreadPool, [this, ByteSize]() { FileMD5Hasher.Update(DataCopy.GetData(), ByteSize); });
+			}
+			if (Config.bCalculateSHA256)
+			{
+				HashFutures[3] = Async(EAsyncExecution::ThreadPool, [this, ByteSize]() { FileSHA256Hasher.Update(DataCopy.GetData(), ByteSize); });
+			}
+		}
+
+		void Final()
+		{
+			Wait();
+			FileSHA1Hasher.Final();
+			if (Config.bCalculateSHA256)
+			{
+				FileSHA256Hasher.Final();
+			}
+			// We do not call final on MD5 here, as the setting of the value will call final internally.
+		}
+
+		void GetHash(FMD5Hash& MD5Hash)
+		{
+			if (Config.bCalculateMD5)
+			{
+				MD5Hash.Set(FileMD5Hasher);
+			}
+			else
+			{
+				FMemory::Memzero(MD5Hash);
+			}
+		}
+
+		void GetHash(FSHAHash& SHA1Hash)
+		{
+			FileSHA1Hasher.GetHash(SHA1Hash.Hash);
+		}
+
+		void GetHash(FSHA256Signature& SHA256Hash)
+		{
+			if (Config.bCalculateSHA256)
+			{
+				FileSHA256Hasher.GetHash(SHA256Hash.Signature);
+			}
+			else
+			{
+				FMemory::Memzero(SHA256Hash);
+			}
+		}
+
+		void GetMime(FString& MIMEType)
+		{
+			if (Config.bDetectMimeType)
+			{
+				MIMEType = MIMETypeCalc.GetMIMEType();
+			}
+			else
+			{
+				MIMEType.Empty();
+			}
+		}
+
+	private:
+		void Wait()
+		{
+			for (TFuture<void>& HashFuture : HashFutures)
+			{
+				HashFuture.Wait();
+			}
+		}
+
+	private:
+		const FDirectoryBuildStreamerConfig& Config;
+		FSHA1 FileSHA1Hasher;
+		FMIMEType MIMETypeCalc;
+		FMD5 FileMD5Hasher;
+		FSHA256 FileSHA256Hasher;
+		TFuture<void> HashFutures[4];
+		TArray<uint8> DataCopy;
+	};
+
 	class FDataStream
 	{
 	public:
@@ -84,9 +333,6 @@ namespace BuildPatchServices
 		// IBuildStreamer interface end.
 
 		// IDirectoryBuildStreamer interface begin.
-		virtual bool GetFileSpan(uint64 StartingIdx, FFileSpan& FileSpan) const override;
-		virtual TArray<FString> GetEmptyFiles() const override;
-		virtual TArray<FString> GetAllFilenames() const override;
 		virtual uint64 GetBuildSize() const override;
 		virtual TArray<FFileSpan> GetAllFiles() const override;
 		virtual bool HasAborted() const override;
@@ -94,26 +340,15 @@ namespace BuildPatchServices
 
 	private:
 		void ReadData();
-		void AddFile(FFileSpan FileSpan);
-		void AddEmptyFile(FString Filename);
-		void SetFileHash(uint64 StartIdx, FSHA1& FileHash);
-		void ReadInputFileList(TArray<FString>& AllFiles);
-		void StripIgnoredFiles(TArray<FString>& AllFiles);
-		void SetEnumeratedFiles(const TArray<FString>& AllFiles);
 
 	private:
 		const FDirectoryBuildStreamerConfig Config;
 		const FDirectoryBuildStreamerDependencies Dependencies;
 		FDataStream DataStream;
 		mutable FCriticalSection FilesCS;
-		TMap<uint64, FFileSpan> Files;
-		mutable FCriticalSection EnumeratedFilesCS;
-		TArray<FString> EnumeratedFiles;
-		TSet<FString> EmptyFiles;
-		mutable FSHA1 EmptyHasher; // FSHA1 is not const correct.
+		TArray<FFileSpan> Files;
 		TFuture<void> Future;
 		FThreadSafeBool bShouldAbort;
-		FThreadSafeBool bFilesEnumerated;
 	};
 
 	FDataStream::FDataStream()
@@ -206,9 +441,7 @@ namespace BuildPatchServices
 		: Config(MoveTemp(InConfig))
 		, Dependencies(MoveTemp(InDependencies))
 		, bShouldAbort(false)
-		, bFilesEnumerated(false)
 	{
-		EmptyHasher.Final();
 		TFunction<void()> Task = [this]() { ReadData(); };
 		Future = Async(EAsyncExecution::Thread, MoveTemp(Task));
 	}
@@ -216,6 +449,7 @@ namespace BuildPatchServices
 	FDirectoryBuildStreamer::~FDirectoryBuildStreamer()
 	{
 		bShouldAbort = true;
+		DataStream.Clear();
 		Future.Wait();
 	}
 
@@ -224,36 +458,6 @@ namespace BuildPatchServices
 		return DataStream.DequeueData(Buffer, ReqSize, WaitForData);
 	}
 
-	bool FDirectoryBuildStreamer::GetFileSpan(uint64 StartingIdx, FFileSpan& FileSpan) const
-	{
-		FScopeLock ScopeLock(&FilesCS);
-		if (Files.Contains(StartingIdx))
-		{
-			FileSpan = Files[StartingIdx];
-			return true;
-		}
-		return false;
-	}
-
-	TArray<FString> FDirectoryBuildStreamer::GetEmptyFiles() const
-	{
-		FScopeLock ScopeLock(&FilesCS);
-		return EmptyFiles.Array();
-	}
-
-	TArray<FString> FDirectoryBuildStreamer::GetAllFilenames() const
-	{
-		while (!bFilesEnumerated && !bShouldAbort)
-		{
-			FPlatformProcess::Sleep(0.1f);
-		}
-		TArray<FString> AllFiles;
-		EnumeratedFilesCS.Lock();
-		AllFiles.Append(EnumeratedFiles);
-		EnumeratedFilesCS.Unlock();
-		return AllFiles;
-	}
-
 	bool FDirectoryBuildStreamer::IsEndOfData() const
 	{
 		return DataStream.IsEndOfStream() && DataStream.UsedSpace() == 0;
@@ -267,17 +471,9 @@ namespace BuildPatchServices
 
 	TArray<FFileSpan> FDirectoryBuildStreamer::GetAllFiles() const
 	{
-		TArray<FFileSpan> AllFiles;
 		check(DataStream.IsEndOfStream());
 		FScopeLock ScopeLock(&FilesCS);
-		Files.GenerateValueArray(AllFiles);
-		for(const FString& EmptyFile : EmptyFiles)
-		{
-			FFileSpan FileSpan(EmptyFile, 0, 0, false, TEXT(""));
-			EmptyHasher.GetHash(FileSpan.SHAHash.Hash);
-			AllFiles.Add(MoveTemp(FileSpan));
-		}
-		return AllFiles;
+		return Files;
 	}
 
 	bool FDirectoryBuildStreamer::HasAborted() const
@@ -300,64 +496,39 @@ namespace BuildPatchServices
 		// Clear the build stream
 		DataStream.Clear();
 
-		// Enumerate build files
-		TArray<FString> AllFiles;
-		uint32 FileEnumerationStart = FStatsCollector::GetCycles();
-		if (Config.InputListFile.IsEmpty())
-		{
-			Dependencies.FileSystem->FindFilesRecursively(AllFiles, *Config.BuildRoot);
-		}
-		else
-		{
-			ReadInputFileList(AllFiles);
-		}
-		uint32 FileEnumerationEnd = FStatsCollector::GetCycles();
-		uint32 FileEnumerationTime = FileEnumerationEnd - FileEnumerationStart;
-		UE_LOG(LogBuildStreamer, Log, TEXT("Enumerated %d files in %s"), AllFiles.Num(), *FPlatformTime::PrettyTime(FStatsCollector::CyclesToSeconds(FileEnumerationTime)));
-
-		// Remove the files that appear in the ignore list
-		AllFiles.Sort();
-		StripIgnoredFiles(AllFiles);
-
-		// Preserve our sorted, stripped list of files, so it can be retrieved later via GetAllFilenames().
-		SetEnumeratedFiles(AllFiles);
-
-		// Track file hashes
-		FSHA1 FileHash;
-
 		// Allocate our file read buffer
 		uint8* FileReadBuffer = new uint8[BuildPatchServices::FileBufferSize];
 
-		for (FString& SourceFile : AllFiles)
+		for (const FString& BuildFile : Config.BuildFiles)
 		{
+			FAsyncFileHahser Hasher(Config);
+
 			if (bShouldAbort)
 			{
 				break;
 			}
 
 			// Read the file
+			const FString SourceFile = Config.BuildRoot / BuildFile;
 			FStatsCollector::AccumulateTimeBegin(TempValue);
 			TUniquePtr<FArchive> FileReader = Dependencies.FileSystem->CreateFileReader(*SourceFile);
+			uint32 LastError = FileReader.IsValid() ? 0 : FPlatformMisc::GetLastError();
 			bool bIsUnixExe = IsUnixExecutable(*SourceFile);
-			FString SymlinkTarget = GetSymlinkTarget(*SourceFile);
+			FString SymlinkTarget = Config.bFollowSymlinks ? TEXT("") : GetSymlinkTarget(*SourceFile);
 			FStatsCollector::AccumulateTimeEnd(StatFileOpenTime, TempValue);
-			// Not being able to load a required file from the build would be fatal, hard fault.
-			if(!FileReader.IsValid())
+			// Not being able to load a required file from the build is a hard error.
+			if (SymlinkTarget.IsEmpty() && !FileReader.IsValid())
 			{
-				UE_LOG(LogBuildStreamer, Error, TEXT("Failed to open file reader for %s"), *SourceFile);
+				UE_LOG(LogBuildStreamer, Error, TEXT("Failed to open file reader for %s (%u)"), *SourceFile, LastError);
 				bShouldAbort = true;
 				break;
 			}
 
-			// Make SourceFile the format we want it in and start a new file.
-			FPaths::MakePathRelativeTo(SourceFile, *(Config.BuildRoot + TEXT("/")));
-			int64 FileSize = FileReader->TotalSize();
 			// Process files that have bytes.
+			uint64 FileStartIdx = DataStream.TotalDataPushed();
+			int64 FileSize = SymlinkTarget.IsEmpty() ? FileReader->TotalSize() : 0;
 			if (FileSize > 0)
 			{
-				FileHash.Reset();
-				uint64 FileStartIdx = DataStream.TotalDataPushed();
-				AddFile(FFileSpan(SourceFile, FileSize, FileStartIdx, bIsUnixExe, SymlinkTarget));
 				while (!FileReader->AtEnd() && !bShouldAbort)
 				{
 					// Read data file
@@ -368,9 +539,9 @@ namespace BuildPatchServices
 					FStatsCollector::AccumulateTimeEnd(StatFileReadTime, TempValue);
 					FStatsCollector::Accumulate(StatTotalDataRead, ReadLen);
 
-					// Hash data file
+					// Kick off hashing.
 					FStatsCollector::AccumulateTimeBegin(TempValue);
-					FileHash.Update(FileReadBuffer, ReadLen);
+					Hasher.AsyncUpdate(FileReadBuffer, ReadLen);
 					FStatsCollector::AccumulateTimeEnd(StatFileHashTime, TempValue);
 
 					// Copy into data stream
@@ -382,14 +553,19 @@ namespace BuildPatchServices
 					FStatsCollector::Set(StatDataAccessSpeed, *StatTotalDataRead / FStatsCollector::CyclesToSeconds(*StatFileOpenTime + *StatFileReadTime));
 					FStatsCollector::Set(StatPotentialThroughput, *StatTotalDataRead / FStatsCollector::CyclesToSeconds(*StatFileOpenTime + *StatFileReadTime + *StatFileHashTime));
 				}
-				FileHash.Final();
-				SetFileHash(FileStartIdx, FileHash);
-			}
-			// Special case zero byte files.
-			else if (FileSize == 0)
-			{
-				AddEmptyFile(SourceFile);
 			}
+			// Finalise hashing.
+			FStatsCollector::AccumulateTimeBegin(TempValue);
+			Hasher.Final();
+			FStatsCollector::AccumulateTimeEnd(StatFileHashTime, TempValue);
+			// Save the file span.
+			FFileSpan FileSpan(BuildFile, FileSize, FileStartIdx, bIsUnixExe, SymlinkTarget);
+			Hasher.GetHash(FileSpan.MD5Hash);
+			Hasher.GetHash(FileSpan.SHA1Hash);
+			Hasher.GetHash(FileSpan.SHA256Hash);
+			Hasher.GetMime(FileSpan.MIMEType);
+			FScopeLock ScopeLock(&FilesCS);
+			Files.Add(FileSpan);
 			FileReader->Close();
 		}
 
@@ -400,105 +576,6 @@ namespace BuildPatchServices
 		delete[] FileReadBuffer;
 	}
 
-	void FDirectoryBuildStreamer::AddFile(FFileSpan FileSpan)
-	{
-		FScopeLock ScopeLock(&FilesCS);
-		Files.Add(FileSpan.StartIdx, MoveTemp(FileSpan));
-	}
-
-	void FDirectoryBuildStreamer::AddEmptyFile(FString Filename)
-	{
-		FScopeLock ScopeLock(&FilesCS);
-		EmptyFiles.Add(MoveTemp(Filename));
-	}
-
-	void FDirectoryBuildStreamer::SetFileHash(uint64 StartIdx, FSHA1& FileHash)
-	{
-		FScopeLock ScopeLock(&FilesCS);
-		FileHash.GetHash(Files[StartIdx].SHAHash.Hash);
-	}
-
-	void FDirectoryBuildStreamer::ReadInputFileList(TArray<FString>& AllFiles)
-	{
-		FString InputFileList;
-		FFileHelper::LoadFileToString(InputFileList, *Config.InputListFile);
-
-		TArray<FString> InputFiles;
-		InputFileList.ParseIntoArrayLines(InputFiles, true);
-
-		for (FString& InputFile : InputFiles)
-		{
-			InputFile.TrimStartAndEndInline();
-			if (InputFile.Len() > 0)
-			{
-				FString FullInputFile = Config.BuildRoot / InputFile;
-				FPaths::NormalizeFilename(FullInputFile);
-				AllFiles.Add(FullInputFile);
-			}
-		}
-	}
-
-	void FDirectoryBuildStreamer::StripIgnoredFiles(TArray<FString>& AllFiles)
-	{
-		struct FRemoveMatchingStrings
-		{
-			const TSet<FString>& IgnoreList;
-			FRemoveMatchingStrings(const TSet<FString>& InIgnoreList)
-				: IgnoreList(InIgnoreList) {}
-
-			bool operator()(const FString& RemovalCandidate) const
-			{
-				const bool bRemove = IgnoreList.Contains(RemovalCandidate);
-				if (bRemove)
-				{
-					UE_LOG(LogBuildStreamer, Log, TEXT("    - %s"), *RemovalCandidate);
-				}
-				return bRemove;
-			}
-		};
-
-		UE_LOG(LogBuildStreamer, Log, TEXT("Stripping ignorable files"));
-		const int32 OriginalNumFiles = AllFiles.Num();
-		FString IgnoreFileList;
-		FFileHelper::LoadFileToString(IgnoreFileList, *Config.IgnoreListFile);
-		TArray< FString > IgnoreFiles;
-		IgnoreFileList.ParseIntoArrayLines(IgnoreFiles, true);
-
-		// Normalize all paths first
-		for (FString& Filename : AllFiles)
-		{
-			FPaths::NormalizeFilename(Filename);
-		}
-		for (FString& Filename : IgnoreFiles)
-		{
-			int32 TabLocation = Filename.Find(TEXT("\t"), ESearchCase::CaseSensitive);
-			if (TabLocation != INDEX_NONE)
-			{
-				// Strip tab separated timestamp if it exists
-				Filename.LeftInline(TabLocation, EAllowShrinking::No);
-			}
-			Filename = Config.BuildRoot / Filename;
-			FPaths::NormalizeFilename(Filename);
-		}
-
-		// Convert ignore list to set
-		TSet<FString> IgnoreSet(MoveTemp(IgnoreFiles));
-
-		// Filter file list
-		FRemoveMatchingStrings FileFilter(IgnoreSet);
-		AllFiles.RemoveAll(FileFilter);
-
-		const int32 NewNumFiles = AllFiles.Num();
-		UE_LOG(LogBuildStreamer, Log, TEXT("Stripped %d ignorable file(s)"), (OriginalNumFiles - NewNumFiles));
-	}
-
-	void FDirectoryBuildStreamer::SetEnumeratedFiles(const TArray<FString>& AllFiles)
-	{
-		FScopeLock ScopeLock(&EnumeratedFilesCS);
-		EnumeratedFiles = AllFiles;
-		bFilesEnumerated = true;
-	}
-
 	IDirectoryBuildStreamer* FBuildStreamerFactory::Create(FDirectoryBuildStreamerConfig Config, FDirectoryBuildStreamerDependencies Dependencies)
 	{
 		check(Dependencies.StatsCollector != nullptr);
diff --git a/Engine/Source/Runtime/Online/BuildPatchServices/Private/Generation/FileAttributesParser.cpp b/Engine/Source/Runtime/Online/BuildPatchServices/Private/Generation/FileAttributesParser.cpp
--- a/Engine/Source/Runtime/Online/BuildPatchServices/Private/Generation/FileAttributesParser.cpp
+++ b/Engine/Source/Runtime/Online/BuildPatchServices/Private/Generation/FileAttributesParser.cpp
@@ -8,12 +8,6 @@ DEFINE_LOG_CATEGORY(LogFileAttributesParser);
 
 namespace BuildPatchServices
 {
-	FFileAttributes::FFileAttributes()
-		: bReadOnly(false)
-		, bCompressed(false)
-		, bUnixExecutable(false)
-	{}
-
 	class FFileAttributesParserImpl
 		: public FFileAttributesParser
 	{
@@ -74,6 +68,8 @@ namespace BuildPatchServices
 
 	bool FFileAttributesParserImpl::FileAttributesMetaToMap(const FString& AttributesList, TMap<FString, FFileAttributes>& FileAttributesMap)
 	{
+		const TCHAR RegexId = TEXT('R');
+		const TCHAR WildcardId = TEXT('W');
 		const TCHAR Quote = TEXT('\"');
 		const TCHAR EOFile = TEXT('\0');
 		const TCHAR EOLine = TEXT('\n');
@@ -85,7 +81,7 @@ namespace BuildPatchServices
 		while (*CharPtr != EOFile)
 		{
 			// Parse filename
-			while (*CharPtr != Quote && *CharPtr != EOFile){ ++CharPtr; }
+			while (FChar::IsWhitespace(*CharPtr) && *CharPtr != EOFile) { ++CharPtr; }
 			if (*CharPtr == EOFile)
 			{
 				if (!FoundFilename)
@@ -95,22 +91,50 @@ namespace BuildPatchServices
 				}
 				break;
 			}
-			const TCHAR* FilenameStart = ++CharPtr;
-			while (*CharPtr != Quote && *CharPtr != EOFile && *CharPtr != EOLine){ ++CharPtr; }
+			// make sure file starts with a quote
+			bool bUsualName = *CharPtr == Quote;
+			// make sure patterns starts with R" or W"
+			bool bRegexName = *CharPtr == RegexId && *(CharPtr + 1) == Quote;
+			bool bWildcardName = *CharPtr == WildcardId && *(CharPtr + 1) == Quote;
+			const TCHAR* FilenameStart = nullptr;
+			if (bUsualName)
+			{
+				// skip "
+				FilenameStart = ++CharPtr;
+			}
+			else if (bRegexName || bWildcardName)
+			{
+				// keep R" and W" in filename
+				FilenameStart = CharPtr;
+				// skip R" and W"
+				CharPtr += 2;
+			}
+			else
+			{
+				UE_LOG(LogFileAttributesParser, Error, TEXT("Did not find opening quote for filename. Pos:%d"), CharPtr - *AttributesList);
+				return false;
+			}
+			while (*CharPtr != Quote && *CharPtr != EOFile && *CharPtr != EOLine) { ++CharPtr; }
 			// Check we didn't run out of file
 			if (*CharPtr == EOFile)
 			{
-				UE_LOG(LogFileAttributesParser, Error, TEXT("Unexpected end of file before next quote! Pos:%" SIZE_T_FMT), CharPtr - *AttributesList);
+				UE_LOG(LogFileAttributesParser, Error, TEXT("Unexpected end of file before next quote! File names must be quoted. Pos:%" SIZE_T_FMT), CharPtr - *AttributesList);
 				return false;
 			}
 			// Check we didn't run out of line
-			if(*CharPtr == EOLine)
+			if (*CharPtr == EOLine)
 			{
-				UE_LOG(LogFileAttributesParser, Error, TEXT("Unexpected end of line before next quote! Pos:%" SIZE_T_FMT), CharPtr - *AttributesList);
+				UE_LOG(LogFileAttributesParser, Error, TEXT("Unexpected end of line before next quote! File names must be quoted. Pos:%" SIZE_T_FMT), CharPtr - *AttributesList);
 				return false;
 			}
 			// Save positions
 			const TCHAR* FilenameEnd = CharPtr++;
+			if (bRegexName || bWildcardName)
+			{
+				// we keep the closing " for pattern matchers
+				FilenameEnd++;
+				CharPtr++;
+			}
 			const TCHAR* AttributesStart = CharPtr;
 			// Parse keywords
 			while (*CharPtr != Quote && *CharPtr != EOFile && *CharPtr != EOLine){ ++CharPtr; }
@@ -124,11 +148,17 @@ namespace BuildPatchServices
 			// Save position
 			const TCHAR* EndOfLine = CharPtr;
 			// Grab info
-			FString Filename = FString::ConstructFromPtrSize(FilenameStart, FilenameEnd - FilenameStart).Replace(TEXT("\\"), TEXT("/"));
+			FString Filename = FString(FilenameEnd - FilenameStart, FilenameStart);
 			FFileAttributes& FileAttributes = FileAttributesMap.FindOrAdd(Filename);
 			TArray<FString> AttributeParamsArray;
 			FString AttributeParams = FString::ConstructFromPtrSize(AttributesStart, EndOfLine - AttributesStart);
 			AttributeParams.ParseIntoArrayWS(AttributeParamsArray);
+			if (AttributeParamsArray.Num() == 0)
+			{
+				UE_LOG(LogFileAttributesParser, Error, TEXT("No attributes have been set for %s"), *Filename);
+				Successful = false;
+			}
+
 			for (const FString& AttributeParam : AttributeParamsArray)
 			{
 				FString Key, Value;
@@ -146,6 +176,12 @@ namespace BuildPatchServices
 					Successful = false;
 				}
 			}
+			FileAttributes.InstallTags.Sort(TLess<FString>());
+		}
+
+		if (!Successful)
+		{
+			UE_LOG(LogFileAttributesParser, Error, TEXT("Please provide the correct attribute. If no attribute is needed, remove the file from the input."));
 		}
 
 		return Successful;
diff --git a/Engine/Source/Runtime/Online/BuildPatchServices/Private/Generation/FileAttributesParser.h b/Engine/Source/Runtime/Online/BuildPatchServices/Private/Generation/FileAttributesParser.h
--- a/Engine/Source/Runtime/Online/BuildPatchServices/Private/Generation/FileAttributesParser.h
+++ b/Engine/Source/Runtime/Online/BuildPatchServices/Private/Generation/FileAttributesParser.h
@@ -2,17 +2,90 @@
 #pragma once
 
 #include "CoreMinimal.h"
+#include "Generation/RegexFind.h"
 #include "HAL/PlatformFileManager.h"
+#include "Misc/Paths.h"
 
 namespace BuildPatchServices
 {
 	struct FFileAttributes
 	{
-		bool bReadOnly;
-		bool bCompressed;
-		bool bUnixExecutable;
+		bool bReadOnly = false;
+		bool bCompressed = false;
+		bool bUnixExecutable = false;
 		TSet<FString> InstallTags;
-		FFileAttributes();
+	};
+
+	class FFileAttributesMap
+	{
+	public:
+		FFileAttributesMap() = default;
+
+		FFileAttributesMap(TMap<FString, FFileAttributes> ExtMap)
+		{
+			Emplace(MoveTemp(ExtMap));
+		}
+
+		void Emplace(TMap<FString, FFileAttributes> ExtMap)
+		{
+			Map = MoveTemp(ExtMap);
+			SortPatterns();
+		}
+
+		const TMap<FString, FFileAttributes>& Raw() const
+		{
+			return Map;
+		}
+
+		const TMap<FString, FFileAttributes>& RawPatterns() const
+		{
+			return PatternMap;
+		}
+
+		FFileAttributes& FindExactOrAdd(const FString& Name)
+		{
+			return Map.FindOrAdd(Name);
+		}
+
+		FFileAttributes Find(const FString& Name) const
+		{
+			FFileAttributes Out;
+			TArray<const FFileAttributes*> AllAttributes = RegexFind(Name, PatternMap);
+			if (const FFileAttributes* Attributes = Map.Find(Name))
+			{
+				AllAttributes.Add(Attributes);
+			}
+			for (const FFileAttributes* Attributes : AllAttributes)
+			{
+				Out.bReadOnly |= Attributes->bReadOnly;
+				Out.bCompressed |= Attributes->bCompressed;
+				Out.bUnixExecutable |= Attributes->bUnixExecutable;
+				Out.InstallTags.Append(Attributes->InstallTags);
+			}
+			return Out;
+		}
+
+	private:
+		void SortPatterns()
+		{
+			TMap<FString, FFileAttributes> TmpMap = MoveTemp(Map);
+			for (TPair<FString, FFileAttributes>& Item : TmpMap)
+			{
+				if (FRegexFinder::IsPattern(Item.Key))
+				{
+					PatternMap.Add(MoveTemp(Item));
+				}
+				else
+				{
+					FPaths::NormalizeFilename(Item.Key);
+					Map.Add(MoveTemp(Item));
+				}
+			}
+		}
+
+	private:
+		TMap<FString, FFileAttributes> Map;
+		TMap<FString, FFileAttributes> PatternMap;
 	};
 
 	class FFileAttributesParser
diff --git a/Engine/Source/Runtime/Online/BuildPatchServices/Private/Generation/ManifestBuildStreamer.cpp b/Engine/Source/Runtime/Online/BuildPatchServices/Private/Generation/ManifestBuildStreamer.cpp
--- a/Engine/Source/Runtime/Online/BuildPatchServices/Private/Generation/ManifestBuildStreamer.cpp
+++ b/Engine/Source/Runtime/Online/BuildPatchServices/Private/Generation/ManifestBuildStreamer.cpp
@@ -27,7 +27,8 @@
 #include "Installer/MessagePump.h"
 #include "BuildPatchProgress.h"
 
-DECLARE_LOG_CATEGORY_CLASS(LogManifestBuildStreamer, Log, All);
+DECLARE_LOG_CATEGORY_EXTERN(LogManifestBuildStreamer, Log, All);
+DEFINE_LOG_CATEGORY(LogManifestBuildStreamer);
 
 namespace BuildPatchServices
 {
@@ -100,7 +101,8 @@ namespace BuildPatchServices
 				}
 				else
 				{
-					UE_LOG(LogManifestBuildStreamer, Fatal, TEXT("Could not get chunk data (%s)."), *NextData.Guid.ToString());
+					UE_LOG(LogManifestBuildStreamer, Error, TEXT("Could not get chunk data (%s)."), *NextData.Guid.ToString());
+					break;
 				}
 			}
 		}
@@ -209,4 +211,4 @@ namespace BuildPatchServices
 		ManifestBuildStreamer->Initialise();
 		return ManifestBuildStreamer;
 	}
-}
\ No newline at end of file
+}
diff --git a/Engine/Source/Runtime/Online/BuildPatchServices/Private/Generation/ManifestBuilder.cpp b/Engine/Source/Runtime/Online/BuildPatchServices/Private/Generation/ManifestBuilder.cpp
--- a/Engine/Source/Runtime/Online/BuildPatchServices/Private/Generation/ManifestBuilder.cpp
+++ b/Engine/Source/Runtime/Online/BuildPatchServices/Private/Generation/ManifestBuilder.cpp
@@ -1,10 +1,13 @@
 // Copyright Epic Games, Inc. All Rights Reserved.
 
 #include "Generation/ManifestBuilder.h"
-#include "Misc/Paths.h"
+
 #include "Algo/Accumulate.h"
-#include "Data/ManifestData.h"
+#include "BuildPatchServicesModule.h"
+#include "BuildPatchUtil.h"
 #include "Common/FileSystem.h"
+#include "Data/ManifestData.h"
+#include "Misc/Paths.h"
 
 DECLARE_LOG_CATEGORY_EXTERN(LogManifestBuilder, Log, All);
 DEFINE_LOG_CATEGORY(LogManifestBuilder);
@@ -42,7 +45,8 @@ namespace BuildPatchServices
 
 		virtual void AddChunkMatch(const FGuid& ChunkGuid, const FBlockStructure& Structure) override;
 		virtual bool FinalizeData(const TArray<FFileSpan>& FileSpans, TArray<FChunkInfo> ChunkInfo) override;
-		virtual bool SaveToFile(const FString& Filename) override;
+		virtual bool SaveToFile(const FString& Filename, FSHAHash* OutSHA1Hash, FMD5Hash* OutMD5Hash) override;
+		virtual FBuildPatchAppManifest& GetManifest() const override;
 
 	private:
 		TArray<FChunkPart> GetChunkPartsForFile(uint64 StartIdx, uint64 Size, TSet<FGuid>& ReferencedChunks);
@@ -52,7 +56,9 @@ namespace BuildPatchServices
 		const FManifestBuilderConfig Config;
 
 		FBuildPatchAppManifestRef Manifest;
-		TMap<FString, FFileAttributes> FileAttributesMap;
+		FGuid ManifestEncryptionSecretId;
+		TArray<uint8> ManifestEncryptionSecretKey;
+		FFileAttributesMap FileAttributesMap;
 		FBlockStructure BuildStructureAdded;
 
 		TMap<FGuid, TArray<FBlockStructure>> AllMatches;
@@ -78,6 +84,11 @@ namespace BuildPatchServices
 		Manifest->ManifestMeta.PrereqName = InDetails.PrereqName;
 		Manifest->ManifestMeta.PrereqPath = InDetails.PrereqPath;
 		Manifest->ManifestMeta.PrereqArgs = InDetails.PrereqArgs;
+		Manifest->ManifestMeta.UninstallActionPath = InDetails.UninstallActionPath;
+		Manifest->ManifestMeta.UninstallActionArgs = InDetails.UninstallActionArgs;
+		ManifestEncryptionSecretId = InDetails.EncryptionSecretId;
+		ManifestEncryptionSecretKey = InDetails.EncryptionSecretKey;
+
 		for (const auto& CustomField : InDetails.CustomFields)
 		{
 			EVariantTypes VarType = CustomField.Value.GetType();
@@ -119,11 +130,14 @@ namespace BuildPatchServices
 		// For each file create its manifest.
 		for (const FFileSpan& FileSpan : FileSpans)
 		{
-			FFileAttributes FileAttributes = FileAttributesMap.FindRef(FileSpan.Filename);
+			FFileAttributes FileAttributes = FileAttributesMap.Find(FileSpan.Filename);
 			Manifest->FileManifestList.FileList.AddDefaulted();
 			FFileManifest& FileManifest = Manifest->FileManifestList.FileList.Last();
 			FileManifest.Filename = FileSpan.Filename;
-			FMemory::Memcpy(FileManifest.FileHash.Hash, FileSpan.SHAHash.Hash, FSHA1::DigestSize);
+			FileManifest.MD5Hash = FileSpan.MD5Hash;
+			FileManifest.SHA1Hash = FileSpan.SHA1Hash;
+			FileManifest.SHA256Hash = FileSpan.SHA256Hash;
+			FileManifest.MIMEType = FileSpan.MIMEType;
 			FileManifest.InstallTags = FileAttributes.InstallTags.Array();
 			FileManifest.SymlinkTarget = FileSpan.SymlinkTarget;
 			if (FileAttributes.bReadOnly)
@@ -205,20 +219,43 @@ namespace BuildPatchServices
 			return false;
 		}
 
+		if (Manifest->ManifestMeta.FeatureLevel <= EFeatureLevel::UsesRuntimeGeneratedBuildId)
+		{
+			Manifest->ManifestMeta.BuildId = FBuildPatchUtils::GetBackwardsCompatibleBuildId(Manifest->ManifestMeta);
+		}
+
+		if (ManifestEncryptionSecretId.IsValid())
+		{
+			if (!Manifest->EncryptData(ManifestEncryptionSecretId, ManifestEncryptionSecretKey))
+			{
+				UE_LOG(LogManifestBuilder, Error, TEXT("Generated manifest failed to encrypt properly."));
+				return false;
+			}
+		}
+
 		// Everything seems fine.
 		bFinalizationSucceeded = true;
 		return true;
 	}
 
-	bool FManifestBuilder::SaveToFile(const FString& Filename)
+	bool FManifestBuilder::SaveToFile(const FString& Filename, FSHAHash* OutSHA1Hash, FMD5Hash* OutMD5Hash)
 	{
 		// Check that previous validation from FinalizeData succeeded, the error should have been picked up.
 		bool bSuccess = false;
 		if (bFinalizationSucceeded)
 		{
 			const FString TmpFilename = Filename + TEXT("tmp");
-			bSuccess = Manifest->SaveToFile(TmpFilename, Manifest->ManifestMeta.FeatureLevel);
-			bSuccess = FileSystem->MoveFile(*Filename, *TmpFilename) && bSuccess;
+			bSuccess = Manifest->SaveToFile(TmpFilename, Manifest->ManifestMeta.FeatureLevel, OutSHA1Hash, OutMD5Hash);
+			if (!bSuccess)
+			{
+				UE_LOG(LogManifestBuilder, Error, TEXT("Could not export manifest to %s."), *TmpFilename);
+				return bSuccess;
+			}
+			bSuccess = FileSystem->MoveFile(*Filename, *TmpFilename);
+			if (!bSuccess)
+			{
+				UE_LOG(LogManifestBuilder, Error, TEXT("Could not move manifest from %s to %s."), *TmpFilename, *Filename);
+			}
 		}
 		else
 		{
@@ -227,9 +264,20 @@ namespace BuildPatchServices
 		return bSuccess;
 	}
 
+	FBuildPatchAppManifest& FManifestBuilder::GetManifest() const
+	{
+		check(bFinalizationSucceeded);
+		return *Manifest;
+	}
+
 	TArray<FChunkPart> FManifestBuilder::GetChunkPartsForFile(uint64 FileStart, uint64 FileSize, TSet<FGuid>& ReferencedChunks)
 	{
 		TArray<FChunkPart> FileChunkParts;
+		// Short circuit empty files
+		if (FileSize == 0)
+		{
+			return FileChunkParts;
+		}
 		// Collect all matching blocks.
 		TArray<FFileBlock> MatchingBlocks;
 		uint64 FileEnd = FileStart + FileSize;
diff --git a/Engine/Source/Runtime/Online/BuildPatchServices/Private/Generation/ManifestBuilder.h b/Engine/Source/Runtime/Online/BuildPatchServices/Private/Generation/ManifestBuilder.h
--- a/Engine/Source/Runtime/Online/BuildPatchServices/Private/Generation/ManifestBuilder.h
+++ b/Engine/Source/Runtime/Online/BuildPatchServices/Private/Generation/ManifestBuilder.h
@@ -33,10 +33,18 @@ namespace BuildPatchServices
 		FString PrereqPath;
 		// The command line arguments for the prerequisites installer
 		FString PrereqArgs;
+		// The path to the uninstall custom action executable.
+		FString UninstallActionPath;
+		// The arguments to the uninstall custom action executable.
+		FString UninstallActionArgs;
 		// Map of custom fields to add to the manifest
 		TMap<FString, FVariant> CustomFields;
 		// Map of file attributes
-		TMap<FString, FFileAttributes> FileAttributesMap;
+		FFileAttributesMap FileAttributesMap;
+		// The encryption secret key to use for this manifest.
+		TArray<uint8> EncryptionSecretKey;
+		// The ID of the encryption secret key used by this manifest.
+		FGuid EncryptionSecretId;
 	};
 
 	struct FManifestBuilderConfig
@@ -58,7 +66,8 @@ namespace BuildPatchServices
 		virtual ~IManifestBuilder() = default;
 		virtual void AddChunkMatch(const FGuid& ChunkGuid, const FBlockStructure& Structure) = 0;
 		virtual bool FinalizeData(const TArray<FFileSpan>& FileSpans, TArray<FChunkInfo> ChunkInfo) = 0;
-		virtual bool SaveToFile(const FString& Filename) = 0;
+		virtual bool SaveToFile(const FString& Filename, FSHAHash* OutSHA1Hash, FMD5Hash* OutMD5Hash) = 0;
+		virtual FBuildPatchAppManifest& GetManifest() const = 0;
 	};
 
 	typedef TSharedRef<IManifestBuilder> IManifestBuilderRef;
diff --git a/Engine/Source/Runtime/Online/BuildPatchServices/Private/Generation/PackageChunkData.cpp b/Engine/Source/Runtime/Online/BuildPatchServices/Private/Generation/PackageChunkData.cpp
--- a/Engine/Source/Runtime/Online/BuildPatchServices/Private/Generation/PackageChunkData.cpp
+++ b/Engine/Source/Runtime/Online/BuildPatchServices/Private/Generation/PackageChunkData.cpp
@@ -5,6 +5,7 @@
 #include "Async/Async.h"
 #include "Async/Future.h"
 #include "Async/TaskGraphInterfaces.h"
+#include "Misc/ConfigCacheIni.h"
 #include "Misc/FileHelper.h"
 #include "Misc/OutputDeviceRedirector.h"
 #include "Misc/Paths.h"
@@ -16,12 +17,16 @@
 
 #include "Core/Platform.h"
 #include "Common/ChunkDataSizeProvider.h"
+#include "Common/Crypto.h"
 #include "Common/FileSystem.h"
 #include "Common/HttpManager.h"
 #include "Common/SpeedRecorder.h"
+#include "Common/StatsCollector.h"
 #include "Generation/ChunkDatabaseWriter.h"
+#include "Installer/Statistics/CloudChunkSourceStat.h"
 #include "Installer/Statistics/CloudChunkSourceStatistics.h"
 #include "Installer/Statistics/DownloadServiceStatistics.h"
+#include "Installer/Statistics/FileConstructorStatistics.h"
 #include "Installer/Statistics/FileOperationTracker.h"
 #include "Installer/Statistics/MemoryChunkStoreStatistics.h"
 #include "Installer/ChunkEvictionPolicy.h"
@@ -30,6 +35,7 @@
 #include "Installer/DownloadService.h"
 #include "Installer/InstallerAnalytics.h"
 #include "Installer/InstallerError.h"
+#include "Installer/InstallerSharedContext.h"
 #include "Installer/MemoryChunkStore.h"
 #include "Installer/MessagePump.h"
 #include "Installer/OptimisedDelta.h"
@@ -38,7 +44,8 @@
 #include "IBuildManifestSet.h"
 #include "Core/AsyncHelpers.h"
 
-DECLARE_LOG_CATEGORY_CLASS(LogPackageChunkData, Log, All);
+DECLARE_LOG_CATEGORY_EXTERN(LogPackageChunkData, Log, All);
+DEFINE_LOG_CATEGORY(LogPackageChunkData);
 
 // For the output file we'll use pretty json in debug, otherwise condensed.
 #if UE_BUILD_DEBUG
@@ -101,6 +108,36 @@ namespace PackageChunksHelpers
 
 namespace BuildPatchServices
 {
+	namespace ChunkDbHelpers
+	{
+		uint32 GetChunkDbHeaderSize()
+		{
+			static const uint32 ChunkDbHeaderSize = []()
+			{
+				TArray<uint8> HeaderData;
+				FMemoryWriter HeaderWriter(HeaderData);
+				FChunkDatabaseHeader ChunkDbHeader;
+				HeaderWriter << ChunkDbHeader;
+				return HeaderData.Num();
+			}();
+			return ChunkDbHeaderSize;
+		}
+
+		uint32 GetPerChunkEntryHeaderSize()
+		{
+			static const uint32 PerEntryHeaderSize = []()
+			{
+				TArray<uint8> HeaderData;
+				FMemoryWriter HeaderWriter(HeaderData);
+				FChunkDatabaseHeader ChunkDbHeader;
+				ChunkDbHeader.Contents.Add({ FGuid::NewGuid(), 0, 0 });
+				HeaderWriter << ChunkDbHeader;
+				return HeaderData.Num() - GetChunkDbHeaderSize();
+			}();
+			return PerEntryHeaderSize;
+		}
+	}
+
 	class FPackageChunks
 		: public IPackageChunks
 	{
@@ -116,6 +153,8 @@ namespace BuildPatchServices
 		void HandleDownloadComplete(int32 RequestId, const FDownloadRef& Download);
 		void HandleManifestComplete();
 		void HandleManifestSelection(const IOptimisedDelta::FResultValueOrError& ResultValueOrError);
+		bool TagsAreValid(const FBuildPatchAppManifestPtr& ManifestPtr, const TArray<TSet<FString>>& ProvidedTagSetArray, const FString& BuildType);
+		bool TagsAreValid(const FBuildPatchAppManifestPtr& ManifestPtr, const TSet<FString>& ProvidedTagSetArray, const FString& BuildType);
 		void BeginPackageProcess();
 		void OnPackageComplete(bool bInSuccess);
 
@@ -157,25 +196,30 @@ namespace BuildPatchServices
 
 		// Declare constructed systems last, to ensure they are destructed before our data members.
 		TUniquePtr<IPlatform> Platform;
-		TUniquePtr<IHttpManager> HttpManager;
-		TUniquePtr<IFileSystem> FileSystem;
+		TSharedPtr<IHttpManager> HttpManager;
+		TSharedPtr<IFileSystem> FileSystem;
+		TUniquePtr<ICrypto> Crypto;
 		TUniquePtr<IMessagePump> MessagePump;
 		TUniquePtr<IInstallerError> InstallerError;
 		TUniquePtr<ISpeedRecorder> DownloadSpeedRecorder;
+		TUniquePtr<ISpeedRecorder> DiskReadSpeedRecorder;
+		TUniquePtr<ISpeedRecorder> DiskWriteSpeedRecorder;
 		TUniquePtr<IChunkDataSizeProvider> ChunkDataSizeProvider;
-		TUniquePtr<IInstallerAnalytics> InstallerAnalytics;
-		TUniquePtr<IDownloadServiceStatistics> DownloadServiceStatistics;
+		TSharedPtr<IInstallerAnalytics> InstallerAnalytics;
+		TSharedPtr<IDownloadServiceStatistics> DownloadServiceStatistics;
 		TUniquePtr<IDownloadService> DownloadService;
 		TUniquePtr<IChunkReferenceTracker> ChunkReferenceTracker;
 		TUniquePtr<IFileOperationTracker> FileOperationTracker;
-		TUniquePtr<IOptimisedDelta> OptimisedDelta;
+		TSharedPtr<IOptimisedDelta> OptimisedDelta;
 		FBuildPatchProgress BuildProgress;
 		TUniquePtr<IMemoryChunkStoreStatistics> MemoryChunkStoreStatistics;
 		TUniquePtr<ICloudChunkSourceStatistics> CloudChunkSourceStatistics;
+		TUniquePtr<IFileConstructorStatistics> FileConstructorStatistics;
 		TUniquePtr<IChunkDataSerialization> ChunkDataSerialization;
 		TUniquePtr<IChunkEvictionPolicy> MemoryEvictionPolicy;
 		TUniquePtr<IMemoryChunkStore> CloudChunkStore;
 		TUniquePtr<IDownloadConnectionCount> DownloadConnectionCount;
+		IBuildInstallerSharedContextPtr CloudSourceSharedContext;
 		TUniquePtr<ICloudChunkSource> CloudChunkSource;
 		TUniquePtr<IChunkDatabaseWriter> ChunkDatabaseWriter;
 	};
@@ -198,13 +242,16 @@ namespace BuildPatchServices
 		, Platform(FPlatformFactory::Create())
 		, HttpManager(FHttpManagerFactory::Create())
 		, FileSystem(FFileSystemFactory::Create())
+		, Crypto(FCryptoFactory::Create())
 		, MessagePump(FMessagePumpFactory::Create())
 		, InstallerError(FInstallerErrorFactory::Create())
 		, DownloadSpeedRecorder(FSpeedRecorderFactory::Create())
-		, ChunkDataSizeProvider(FChunkDataSizeProviderFactory::Create())
+		, DiskReadSpeedRecorder(FSpeedRecorderFactory::Create())
+		, DiskWriteSpeedRecorder(FSpeedRecorderFactory::Create())
+		, ChunkDataSizeProvider(FChunkDataSizeProviderFactory::Create({}))
 		, InstallerAnalytics(FInstallerAnalyticsFactory::Create(nullptr))
 		, DownloadServiceStatistics(FDownloadServiceStatisticsFactory::Create(DownloadSpeedRecorder.Get(), ChunkDataSizeProvider.Get(), InstallerAnalytics.Get()))
-		, DownloadService(FDownloadServiceFactory::Create(HttpManager.Get(), FileSystem.Get(), DownloadServiceStatistics.Get(), InstallerAnalytics.Get()))
+		, DownloadService(FDownloadServiceFactory::Create(HttpManager, FileSystem, DownloadServiceStatistics, InstallerAnalytics))
 		, FileOperationTracker(FFileOperationTrackerFactory::Create(CoreTicker))
 	{
 		// Make sure the cloud chunk source gets the abort signal if an error occurred.
@@ -230,8 +277,12 @@ namespace BuildPatchServices
 		FHttpModule::Get();
 
 		// Kick off Manifest downloads.
-		RequestIdManifestFile = DownloadService->RequestFile(Configuration.ManifestFilePath, DownloadCompleteDelegate, DownloadProgressDelegate);
-		RequestIdPrevManifestFile = DownloadService->RequestFile(Configuration.PrevManifestFilePath, DownloadCompleteDelegate, DownloadProgressDelegate);
+		bShouldRun = bShouldRun && DownloadService->RequestFile(Configuration.ManifestFilePath, DownloadCompleteDelegate, DownloadProgressDelegate, &RequestIdManifestFile);
+		bShouldRun = bShouldRun && DownloadService->RequestFile(Configuration.PrevManifestFilePath, DownloadCompleteDelegate, DownloadProgressDelegate, &RequestIdPrevManifestFile);
+		if (!bShouldRun)
+		{
+			bSuccess = false;
+		}
 
 		// Main timers.
 		double DeltaTime = 0.0;
@@ -317,21 +368,65 @@ namespace BuildPatchServices
 				bSuccess = false;
 				bShouldRun = false;
 			}
+
+			if ((Manifest.IsValid() && !Configuration.bIgnoreManifestFileInvalidTags && !TagsAreValid(Manifest, Configuration.TagSetArray, TEXT("build")))
+				|| (PrevManifest.IsValid() && !Configuration.bIgnorePrevManifestFileInvalidTags && !TagsAreValid(PrevManifest, Configuration.PrevTagSet, TEXT("previous build"))))
+			{
+				bSuccess = false;
+				bShouldRun = false;
+			}
+
 			if (bSuccess)
 			{
 				// Check for delta file, replacing Manifest if we find one
 				FOptimisedDeltaConfiguration OptimisedDeltaConfiguration(Manifest.ToSharedRef());
 				OptimisedDeltaConfiguration.SourceManifest = PrevManifest;
 				OptimisedDeltaConfiguration.DeltaPolicy = Configuration.FeatureLevel >= EFeatureLevel::FirstOptimisedDelta ? EDeltaPolicy::TryFetchContinueWithout : EDeltaPolicy::Skip;
-				OptimisedDeltaConfiguration.CloudDirectories = { FPaths::GetPath(Configuration.ManifestFilePath) };
-				FOptimisedDeltaDependencies OptimisedDeltaDependencies;
+				OptimisedDeltaConfiguration.DeltaFilenameTrailer = Configuration.DeltaFilenameTrailer;
+				TSharedRef<IUriProvider, ESPMode::ThreadSafe> UriProvider = FUriProviderFactory::Create(MessagePump.Get(), { FPaths::GetPath(Configuration.ManifestFilePath) });
+				FOptimisedDeltaDependencies OptimisedDeltaDependencies(MoveTemp(UriProvider));
 				OptimisedDeltaDependencies.DownloadService = DownloadService.Get();
 				OptimisedDeltaDependencies.OnComplete = MakeOptimiseCompleteDelegate(&FPackageChunks::HandleManifestSelection);
-				OptimisedDelta.Reset(FOptimisedDeltaFactory::Create(OptimisedDeltaConfiguration, MoveTemp(OptimisedDeltaDependencies)));
+				OptimisedDelta = FOptimisedDeltaFactory::Create(OptimisedDeltaConfiguration, MoveTemp(OptimisedDeltaDependencies));
 			}
 		}
 	}
 
+	bool FPackageChunks::TagsAreValid(const FBuildPatchAppManifestPtr& ManifestPtr, const TArray<TSet<FString>>& ProvidedTagSetArray, const FString& BuildType)
+	{
+		TSet<FString> ProvidedTags;
+		for (const TSet<FString>& ProvidedTagSet : ProvidedTagSetArray)
+		{
+			ProvidedTags.Append(ProvidedTagSet);
+		}
+		return TagsAreValid(ManifestPtr, ProvidedTags, BuildType);
+	}
+
+	bool FPackageChunks::TagsAreValid(const FBuildPatchAppManifestPtr& ManifestPtr, const TSet<FString>& ProvidedTagSet, const FString& BuildType)
+	{
+		bool bTagsAreValid = true;
+		TArray<FString> UnknownManifestTags;
+
+		TSet<FString> ValidTagSet;
+		ManifestPtr->GetFileTagList(ValidTagSet);
+		for (const FString& ProvidedTag : ProvidedTagSet)
+		{
+			if (!ValidTagSet.Contains(ProvidedTag))
+			{
+				UnknownManifestTags.Add(ProvidedTag);
+				bTagsAreValid = false;
+			}
+		}
+
+		if (!bTagsAreValid)
+		{
+			UE_LOG(LogPackageChunkData, Error, TEXT("The provided tags were invalid for the %s: %s"), *BuildType, *FString::Join(UnknownManifestTags, TEXT(",")));
+			UE_LOG(LogPackageChunkData, Error, TEXT("Valid tags for the %s are: %s"), *BuildType, *FString::Join(ValidTagSet.Array(), TEXT(",")));
+		}
+
+		return bTagsAreValid;
+	}
+
 	void FPackageChunks::HandleManifestSelection(const IOptimisedDelta::FResultValueOrError& ResultValueOrError)
 	{
 		FBuildPatchAppManifestPtr DeltaManifest = ResultValueOrError.IsValid() ? ResultValueOrError.GetValue() : nullptr;
@@ -386,20 +481,12 @@ namespace BuildPatchServices
 			ChunkReferenceTracker.Reset(FChunkReferenceTrackerFactory::Create(PackageChunksHelpers::GetCustomChunkReferences(TagSetArray, Manifest.ToSharedRef())));
 		}
 
-		// Programmatically calculate header file size effects, so that we automatically handle any changes to the header spec.
-		TArray<uint8> HeaderData;
-		FMemoryWriter HeaderWriter(HeaderData);
-		FChunkDatabaseHeader ChunkDbHeader;
-		HeaderWriter << ChunkDbHeader;
-		const uint64 ChunkDbHeaderSize = HeaderData.Num();
-		HeaderWriter.Seek(0);
-		HeaderData.Reset();
-		ChunkDbHeader.Contents.Add({ FGuid::NewGuid(), 0, 0 });
-		HeaderWriter << ChunkDbHeader;
-		const uint64 PerEntryHeaderSize = HeaderData.Num() - ChunkDbHeaderSize;
-
 		// Enumerate the chunks, allocating them to chunk db files.
-		TSet<FGuid> FullDataSet = ChunkReferenceTracker->GetReferencedChunks();
+		const uint64 ChunkDbHeaderSize = ChunkDbHelpers::GetChunkDbHeaderSize();
+		const uint64 PerEntryHeaderSize = ChunkDbHelpers::GetPerChunkEntryHeaderSize();
+		TArray<FGuid> FullDataSetArray;
+		ChunkReferenceTracker->CopyOutOrderedUseList(FullDataSetArray);
+		TSet<FGuid> FullDataSet(MoveTemp(FullDataSetArray));
 		if (FullDataSet.Num() > 0)
 		{
 			// Create the data set for each tagset.
@@ -493,11 +580,28 @@ namespace BuildPatchServices
 				}
 			}
 
-			// Cloud config.
+			CloudSourceSharedContext = FBuildInstallerSharedContextFactory::Create(TEXT("PackageChunks"));
+			CloudSourceSharedContext->PreallocateThreads(1);
+
+			// Cloud config defaults.
 			FCloudSourceConfig CloudSourceConfig({ Configuration.CloudDir });
 			CloudSourceConfig.bBeginDownloadsOnFirstGet = false;
 			CloudSourceConfig.MaxRetryCount = 30;
 
+			// Load max download retry count from engine config.
+			GConfig->GetInt(TEXT("Portal.BuildPatch"), TEXT("ChunkRetries"), CloudSourceConfig.MaxRetryCount, GEngineIni);
+			CloudSourceConfig.MaxRetryCount = FMath::Clamp<int32>(CloudSourceConfig.MaxRetryCount, -1, 1000);
+
+			// Load the number of retries before enabling detailed logging.
+			GConfig->GetInt(TEXT("Portal.BuildPatch"), TEXT("RetriesToEnableVerbose"), CloudSourceConfig.RetriesToEnableVerbose, GEngineIni);
+			CloudSourceConfig.RetriesToEnableVerbose = FMath::Clamp<int32>(CloudSourceConfig.RetriesToEnableVerbose, 0, 1000);
+			CloudSourceConfig.SharedContext = CloudSourceSharedContext.Get();
+
+			TUniquePtr<FStatsCollector> StatsCollector(FStatsCollectorFactory::Create());
+			TUniquePtr<ICloudChunkSourceStat> CloudChunkSourceStat(FCloudChunkSourceStatFactory::Create(StatsCollector.Get()));
+			FChunkDataSerializationConfig ChunkDataSerializationConfig;
+			// UE5 MERGE TODO: ChunkDataSerializationConfig.CloudChunkSourceStat = CloudChunkSourceStat.Get();
+
 			// Create systems.
 			const int32 CloudStoreId = 0;
 			MemoryChunkStoreStatistics.Reset(FMemoryChunkStoreStatisticsFactory::Create(
@@ -506,8 +610,13 @@ namespace BuildPatchServices
 				InstallerAnalytics.Get(),
 				&BuildProgress,
 				FileOperationTracker.Get()));
+			FileConstructorStatistics.Reset(FFileConstructorStatisticsFactory::Create(
+				DiskReadSpeedRecorder.Get(),
+				DiskWriteSpeedRecorder.Get(),
+				&BuildProgress,
+				FileOperationTracker.Get()));
 			ChunkDataSerialization.Reset(FChunkDataSerializationFactory::Create(
-				FileSystem.Get()));
+				FileSystem.Get(), Crypto.Get(), ChunkDataSerializationConfig));
 			MemoryEvictionPolicy.Reset(FChunkEvictionPolicyFactory::Create(
 				ChunkReferenceTracker.Get()));
 			CloudChunkStore.Reset(FMemoryChunkStoreFactory::Create(
@@ -532,13 +641,16 @@ namespace BuildPatchServices
 				FullDataSet));
 
 			// Start an IO output thread which saves all the chunks to the chunkdbs.
+			FChunkDbWriterConfig ChunkDbWriterConfig(ChunkDbFiles);
 			ChunkDatabaseWriter.Reset(FChunkDatabaseWriterFactory::Create(
+				ChunkDbWriterConfig,
 				CloudChunkSource.Get(),
 				FileSystem.Get(),
 				InstallerError.Get(),
 				ChunkReferenceTracker.Get(),
 				ChunkDataSerialization.Get(),
-				ChunkDbFiles,
+				FileConstructorStatistics.Get(),
+				[](float) {},
 				[this](bool bInSuccess) { OnPackageComplete(bInSuccess); }));
 		}
 		// If there are no chunks required to install or patch the provided manifests and tags, we don't create any chunkdbs.
diff --git a/Engine/Source/Runtime/Online/BuildPatchServices/Private/Generation/PackageChunkData.h b/Engine/Source/Runtime/Online/BuildPatchServices/Private/Generation/PackageChunkData.h
--- a/Engine/Source/Runtime/Online/BuildPatchServices/Private/Generation/PackageChunkData.h
+++ b/Engine/Source/Runtime/Online/BuildPatchServices/Private/Generation/PackageChunkData.h
@@ -7,6 +7,12 @@
 
 namespace BuildPatchServices
 {
+	namespace ChunkDbHelpers
+	{
+		uint32 GetChunkDbHeaderSize();
+		uint32 GetPerChunkEntryHeaderSize();
+	};
+
 	class IPackageChunks
 	{
 	public:
diff --git a/Engine/Source/Runtime/Online/BuildPatchServices/Private/IBuildManifestSet.h b/Engine/Source/Runtime/Online/BuildPatchServices/Private/IBuildManifestSet.h
--- a/Engine/Source/Runtime/Online/BuildPatchServices/Private/IBuildManifestSet.h
+++ b/Engine/Source/Runtime/Online/BuildPatchServices/Private/IBuildManifestSet.h
@@ -20,6 +20,12 @@ namespace BuildPatchServices
 		bool bIsRepair;
 	};
 
+	struct FUninstallActionInfo
+	{
+		FString Path;
+		FString Args;
+	};
+
 	/**
 	 * FBuildPatchInstallerAction
 	 * This class decorates an FInstallerAction instance with downcasted functions for internal manifest implementation.
@@ -52,11 +58,13 @@ namespace BuildPatchServices
 		bool IsInstall() const { return InstallerAction.IsInstall(); }
 		bool IsUpdate() const { return InstallerAction.IsUpdate(); }
 		bool IsRepair() const { return InstallerAction.IsRepair(); }
+		bool IsVerifyOnly() const { return InstallerAction.IsVerifyOnly(); }
 		bool IsUninstall() const { return InstallerAction.IsUninstall(); }
 		const TSet<FString>& GetInstallTags() const { return InstallTags; }
 		const TSet<FString>& GetTaggedFiles() const { return TaggedFiles; }
 		const FString& GetInstallSubdirectory() const { return InstallerAction.GetInstallSubdirectory(); }
 		const FString& GetCloudSubdirectory() const { return InstallerAction.GetCloudSubdirectory(); }
+		FString GetUniqueBuildId() const { return InstallManifest.IsValid() ? InstallManifest->GetUniqueBuildId() : FString(); }
 		const FBuildPatchAppManifest& GetCurrentManifest() const { return *CurrentManifest; }
 		const FBuildPatchAppManifest& GetInstallManifest() const { return *InstallManifest; }
 		const FBuildPatchAppManifest* TryGetCurrentManifest() const { return CurrentManifest.Get(); }
@@ -92,16 +100,24 @@ namespace BuildPatchServices
 
 		// Data enumeration
 		virtual void GetReferencedChunks(TSet<FGuid>& DataGuids) const = 0;
+		virtual void GetReferencedChunks(TArray<FGuid>& DataGuids) const = 0;
 		virtual uint64 GetDownloadSize(const FGuid& DataGuid) const = 0;
 		virtual uint64 GetDownloadSize(const TSet<FGuid>& DataGuids) const = 0;
 		virtual FChunkInfo const* GetChunkInfo(const FGuid& DataGuid) const = 0;
 		virtual bool GetChunkShaHash(const FGuid& ChunkGuid, FSHAHash& OutHash) const = 0;
 		virtual FString GetDataFilename(const FGuid& DataGuid) const = 0;
+		virtual FGuid GetDataSecretId(const FGuid& DataGuid) const = 0;
+		virtual FString GetChunkCloudSubdirectory(const FGuid& DataGuid) const = 0;
+		virtual FString GetChunkFilename(const FGuid& DataGuid) const = 0;
+		virtual FString GetChunkUniqueBuildId(const FGuid& DataGuid) const = 0;
+		virtual bool GetChunkFromDeltaStatus(const FGuid& DataGuid) const = 0;
 
 		// File enumeration
 		virtual int32 GetNumExpectedFiles() const = 0;
 		virtual void GetExpectedFiles(TSet<FString>& Filenames) const = 0;
 		virtual void GetOutdatedFiles(const FString& InstallDirectory, TSet<FString>& OutdatedFiles) const = 0;
+		// Gets files from the current manifest that are missing from disk.
+		virtual void GetMissedFiles(const FString& InstallDirectory, TSet<FString>& OutMissedFiles) const = 0;
 		virtual void GetRemovableFiles(TSet<FString>& FilesToRemove) const = 0;
 		virtual const FFileManifest* GetCurrentFileManifest(const FString& BuildFile) const = 0;
 		virtual const FFileManifest* GetNewFileManifest(const FString& BuildFile) const = 0;
@@ -113,8 +129,12 @@ namespace BuildPatchServices
 
 		// Misc actions
 		virtual bool ContainsUpdate() const = 0;
-		virtual bool IsRepairOnly() const = 0;
+		virtual bool IsExclusivelyRepairing() const = 0;
 		virtual bool HasFileAttributes() const = 0;
+		virtual bool IsExclusivelyVerifying() const = 0;
+
+		// Custom Uninstall Action support
+		virtual void GetUninstallActionInfo(TArray<FUninstallActionInfo>& UninstallActionsInfo) const = 0;
 	};
 
 	class FBuildManifestSetFactory
diff --git a/Engine/Source/Runtime/Online/BuildPatchServices/Private/Installer/BuildStatistics.cpp b/Engine/Source/Runtime/Online/BuildPatchServices/Private/Installer/BuildStatistics.cpp
--- a/Engine/Source/Runtime/Online/BuildPatchServices/Private/Installer/BuildStatistics.cpp
+++ b/Engine/Source/Runtime/Online/BuildPatchServices/Private/Installer/BuildStatistics.cpp
@@ -49,6 +49,9 @@ namespace BuildPatchServices
 		virtual double GetDownloadByteSpeed() const override;
 		virtual double GetDiskReadByteSpeed() const override;
 		virtual double GetDiskWriteByteSpeed() const override;
+		virtual int64 GetDownloadedBytes() const override;
+		virtual int64 GetDiskReadBytes() const override;
+		virtual int64 GetDiskWrittenBytes() const override;
 		virtual int32 GetNumVerifyErrors() const override;
 		virtual TMap<EVerifyError, int32> GetVerifyErrorCounts() const override;
 		// IBuildStatistics interface end.
@@ -127,17 +130,33 @@ namespace BuildPatchServices
 
 	double FBuildStatistics::GetDownloadByteSpeed() const
 	{
-		return Installer->GetDownloadSpeedRecorder()->GetAverageSpeed(10.0f);
+		return Installer->GetDownloadSpeed();
 	}
 
 	double FBuildStatistics::GetDiskReadByteSpeed() const
 	{
-		return Installer->GetDiskReadSpeedRecorder()->GetAverageSpeed(10.0f);
+		return Installer->GetDiskReadSpeedRecorder()->GetAverageSpeed(10.0f, true);
 	}
 
 	double FBuildStatistics::GetDiskWriteByteSpeed() const
 	{
-		return Installer->GetDiskWriteSpeedRecorder()->GetAverageSpeed(10.0f);
+		return Installer->GetDiskWriteSpeedRecorder()->GetAverageSpeed(10.0f, true);
+	}
+
+	int64 FBuildStatistics::GetDownloadedBytes() const
+	{
+		return Installer->GetCloudChunkSourceStatistics()->GetDownloadedBytes();
+	}
+
+	int64 FBuildStatistics::GetDiskReadBytes() const
+	{
+		return Installer->GetInstallChunkSourceStatistics()->GetBytesRead()
+		     + Installer->GetVerifierStatistics()->GetBytesVerified();
+	}
+
+	int64 FBuildStatistics::GetDiskWrittenBytes() const
+	{
+		return Installer->GetFileConstructorStatistics()->GetBytesConstructed();
 	}
 
 	int32 FBuildStatistics::GetNumVerifyErrors() const
