p4 integrate //UE5/Dev-Distro-5.5/Engine/Source/Programs/BuildPatchTool/... //UE5/Main/Engine/Source/Programs/BuildPatchTool/...

diff --git a/Engine/Source/Runtime/Online/BuildPatchServices/Private/Installer/ChunkDbChunkSource.cpp b/Engine/Source/Runtime/Online/BuildPatchServices/Private/Installer/ChunkDbChunkSource.cpp
--- a/Engine/Source/Runtime/Online/BuildPatchServices/Private/Installer/ChunkDbChunkSource.cpp
+++ b/Engine/Source/Runtime/Online/BuildPatchServices/Private/Installer/ChunkDbChunkSource.cpp
@@ -62,14 +62,27 @@ DEFINE_LOG_CATEGORY_STATIC(LogChunkDbChunkSource, Log, All);
 namespace BuildPatchServices
 {
 	/**
-	 * Struct holding variables for accessing a chunkdb file's data.
+	 * Class holding variables for accessing a chunkdb file's data.
 	 */
-	struct FChunkDbDataAccess
+	class FChunkDbDataAccess
 	{
+	public:
+		FChunkDbDataAccess(FChunkDatabaseHeader InHeader,
+			FString InChunkDbFileName, TUniquePtr<FArchive> InArchive, IFileHandle* InFileHandler)
+			: Header(MoveTemp(InHeader))
+			, ChunkDbFileName(MoveTemp(InChunkDbFileName))
+			, Archive(MoveTemp(InArchive))
+			, FileHandler(InFileHandler)
+		{
+		}
+
 		FChunkDatabaseHeader Header;
-		TUniquePtr<FArchive> Archive;
 		FString ChunkDbFileName;
 
+		// For truncation, we track when each chunk in the chunkdb is retired.
+		// If the system is enabled, then this array will have the same count as the Header.Contents array.
+		TArray<int32> ChunkRetiresAt;
+
 		// When the reference tracker gets below this watermark, then we know we are done with this file and we can
 		// close/retire it.
 		int32 RetireAt = 0;
@@ -77,12 +90,29 @@ namespace BuildPatchServices
 		// If we're retired then any access is invalid and fatal as the file has been closed and could be deleted.
 		bool bIsRetired = false;
 
+		int64 TotalSize() const
+		{
+			// FArchiveFileReaderGeneric::Size uses a cached value and does not re-evaluate the real size after truncation.
+			return FileHandler ? FileHandler->Size() : 0L;
+		}
+
+		TUniquePtr<UE::TScopeLock<UE::FMutex>> LockArchive(FArchive** OutArchive) const
+		{
+			check(OutArchive != nullptr)
+			// Wrap to unique ptr to return it. Original TScopeLock does not support move semantic.
+			TUniquePtr<UE::TScopeLock<UE::FMutex>> Lock{};
+			*OutArchive = Archive.Get();
+			return Lock;
+		}
+
 		void Retire(IFileSystem* FileSystemIfDeleting, bool bDelete)
 		{
 			bIsRetired = true;
-			FString ArchiveName = Archive->GetArchiveName();
-			Archive.Reset();
-
+			
+			{
+				Archive.Reset();
+				FileHandler = nullptr;
+			}
 			if (bDelete && FileSystemIfDeleting)
 			{
 				if (!FileSystemIfDeleting->DeleteFile(*ChunkDbFileName))
@@ -91,6 +121,17 @@ namespace BuildPatchServices
 				}
 			}
 		}
+
+		void TruncateIfNecessary(IFileSystem* FileSystem, int64 RequiredChunkDbFileSize)
+		{
+			if (FileHandler && FileHandler->Size() > RequiredChunkDbFileSize)
+			{
+				FileHandler->Truncate(RequiredChunkDbFileSize);
+			}
+		}
+	private:
+		TUniquePtr<FArchive> Archive;
+		IFileHandle* FileHandler;
 	};
 
 	// Holds where to get the chunk data from. File + location in file.
@@ -123,6 +164,43 @@ namespace BuildPatchServices
 				}
 			}
 		}
+		virtual void TryTruncateFiles(int32 RemainingChunkCount) override
+		{
+			if (!Configuration.bTruncateChunkDbs || LastRemainingChunkCount == RemainingChunkCount)
+			{
+				return;
+			}
+			LastRemainingChunkCount = RemainingChunkCount;
+			//
+			// We are truncating for critical saving of disk space, so once a chunk has been used from chunkdb, we release space from the chunkdb files if we can.
+			//
+			for (FChunkDbDataAccess& ChunkDbDataAccess : ChunkDbDataAccesses)
+			{
+				if (ChunkDbDataAccess.ChunkRetiresAt.Num() > 0)
+				{
+					const int32 CurrentLastChunk = ChunkDbDataAccess.ChunkRetiresAt.Num() - 1;
+					int32 LastChunkRequired = CurrentLastChunk;
+					for (; LastChunkRequired >= 0; LastChunkRequired--)
+					{
+						if (ChunkDbDataAccess.ChunkRetiresAt[LastChunkRequired] <= RemainingChunkCount)
+						{
+							break;
+						}
+					}
+					if (LastChunkRequired != CurrentLastChunk)
+					{
+						ChunkDbDataAccess.ChunkRetiresAt.SetNum(LastChunkRequired + 1, EAllowShrinking::No);
+						int64 RequiredChunkDbFileSize = 0;
+						if (LastChunkRequired >= 0)
+						{
+							const FChunkLocation& LastChunkLocation = ChunkDbDataAccess.Header.Contents[LastChunkRequired];
+							RequiredChunkDbFileSize = LastChunkLocation.ByteStart + LastChunkLocation.ByteSize;
+						}
+						ChunkDbDataAccess.TruncateIfNecessary(FileSystem, RequiredChunkDbFileSize);
+					}
+				}
+			}
+		}
 		virtual int32 GetChunkUnavailableAt(const FGuid& DataId) const override
 		{
 			// While technically the chunks retire as a result of delete-during-install, we only do this
@@ -140,7 +218,7 @@ namespace BuildPatchServices
 
 		static void LoadChunkDbFiles(
 			const TArray<FString>& ChunkDbFiles, IFileSystem* FileSystem, const TArray<FGuid>& ChunkAccessOrderedList,
-			TArray<FChunkDbDataAccess>& OutChunkDbDataAccesses, TMap<FGuid, FChunkAccessLookup>& OutChunkDbDataAccessLookup, TSet<FGuid>* OutOptionalAvailableStore);
+			TArray<FChunkDbDataAccess>& OutChunkDbDataAccesses, TMap<FGuid, FChunkAccessLookup>& OutChunkDbDataAccessLookup, TSet<FGuid>* OutOptionalAvailableStore, bool bTruncateChunkDbs);
 
 	private:
 
@@ -158,12 +236,13 @@ namespace BuildPatchServices
 
 		// Number of chunks to process in this manifest when we started.
 		int32 OriginalChunkCount = 0;
+		int32 LastRemainingChunkCount = TNumericLimits<int32>::Max();
 	};
 
 	// Read in the headers, evalutate the list of chunks, and determine when we'll be done with our chunk dbs.
 	void FChunkDbChunkSource::LoadChunkDbFiles(
 		const TArray<FString>& ChunkDbFiles, IFileSystem* FileSystem, const TArray<FGuid>& ChunkAccessOrderedList, 
-		TArray<FChunkDbDataAccess>& OutChunkDbDataAccesses, TMap<FGuid, FChunkAccessLookup>& OutChunkDbDataAccessLookup, TSet<FGuid>* OutOptionalAvailableStore)
+		TArray<FChunkDbDataAccess>& OutChunkDbDataAccesses, TMap<FGuid, FChunkAccessLookup>& OutChunkDbDataAccessLookup, TSet<FGuid>* OutOptionalAvailableStore, bool bTruncateChunkDbs)
 	{
 
 		// Allow OS intervention only once.
@@ -173,13 +252,14 @@ namespace BuildPatchServices
 		// Load each chunkdb's TOC to enumerate available chunks.
 		for (const FString& ChunkDbFilename : ChunkDbFiles)
 		{
-			TUniquePtr<FArchive> ChunkDbFile(FileSystem->CreateFileReader(*ChunkDbFilename));
-			if (ChunkDbFile.IsValid())
+			IFileHandle* FileHandler = nullptr;
+			TUniquePtr<FArchive> ChunkDbArchive(FileSystem->CreateFileReaderEx(*ChunkDbFilename, BuildPatchServices::EReadFlags::None, &FileHandler));
+			if (ChunkDbArchive.IsValid())
 			{
 				// Load header.
 				FChunkDatabaseHeader Header;
-				*ChunkDbFile << Header;
-				if (ChunkDbFile->IsError())
+				*ChunkDbArchive << Header;
+				if (ChunkDbArchive->IsError())
 				{
 					GLog->Logf(TEXT("Failed to load chunkdb header for %s"), *ChunkDbFilename);
 				}
@@ -190,10 +270,16 @@ namespace BuildPatchServices
 				else
 				{
 					// Hold on to the handle and header info.
-					FChunkDbDataAccess DataSource;
-					DataSource.Archive = MoveTemp(ChunkDbFile);
-					DataSource.ChunkDbFileName = ChunkDbFilename;
-					DataSource.Header = MoveTemp(Header);
+					FChunkDbDataAccess DataSource(MoveTemp(Header), ChunkDbFilename, MoveTemp(ChunkDbArchive), FileHandler);
+					if (bTruncateChunkDbs)
+					{
+						// Initialise with int32 max, so that any element that does not get filled out later, will evaluate as retired.
+						DataSource.ChunkRetiresAt.AddUninitialized(DataSource.Header.Contents.Num());
+						for (int32& i : DataSource.ChunkRetiresAt)
+						{
+							i = TNumericLimits<int32>::Max();
+						}
+					}
 					OutChunkDbDataAccesses.Add(MoveTemp(DataSource));
 				}
 			}
@@ -212,9 +298,14 @@ namespace BuildPatchServices
 		// Index all chunks to their location info.
 		for (FChunkDbDataAccess& ChunkDbDataAccess : OutChunkDbDataAccesses)
 		{
+			const int64 ChunkDbFileSize = ChunkDbDataAccess.TotalSize();
+
 			for (FChunkLocation& ChunkLocation : ChunkDbDataAccess.Header.Contents)
 			{
-				if (!OutChunkDbDataAccessLookup.Contains(ChunkLocation.ChunkId))
+				// Check this chunk actually exists in the chunkdb, we clamp to avoid unsigned mismatch.
+				const int64 ChunkLastByte = FMath::Clamp<int64>(ChunkLocation.ByteStart + ChunkLocation.ByteSize - 1, 0, TNumericLimits<int64>::Max());
+				const bool bChunkExists = ChunkLocation.ByteSize > 0 && ChunkLastByte <= ChunkDbFileSize;
+				if (bChunkExists && !OutChunkDbDataAccessLookup.Contains(ChunkLocation.ChunkId))
 				{
 					OutChunkDbDataAccessLookup.Add(ChunkLocation.ChunkId, { &ChunkLocation, &ChunkDbDataAccess });
 
@@ -239,14 +330,22 @@ namespace BuildPatchServices
 				continue;
 			}
 
-			FileLastSeenAt.FindOrAdd(SourceForGuid->DbFile->Archive->GetArchiveName()) = GuidIndex;
+			FileLastSeenAt.FindOrAdd(SourceForGuid->DbFile->ChunkDbFileName) = GuidIndex;
+
+			if (bTruncateChunkDbs)
+			{
+				const TArray<FChunkLocation>& ChunkDBContents = SourceForGuid->DbFile->Header.Contents;
+				const int32 IndexOfChunk = SourceForGuid->Location - &ChunkDBContents[0];
+				check(ChunkDBContents.IsValidIndex(IndexOfChunk)); // programmer error if this fires
+				SourceForGuid->DbFile->ChunkRetiresAt[IndexOfChunk] = ChunkAccessOrderedList.Num() - (GuidIndex + 1);
+			}
 		}
 
 		for (FChunkDbDataAccess& ChunkDbDataAccess : OutChunkDbDataAccesses)
 		{
 			// Default to retired immediately (this is functionally after the first file completes)
 			ChunkDbDataAccess.RetireAt = ChunkAccessOrderedList.Num();
-			int32* LastAt = FileLastSeenAt.Find(ChunkDbDataAccess.Archive->GetArchiveName());
+			int32* LastAt = FileLastSeenAt.Find(ChunkDbDataAccess.ChunkDbFileName);
 			if (LastAt)
 			{
 				// The reference stack gets popped rather than advanced, so we need to reverse the ordering.
@@ -266,7 +365,7 @@ namespace BuildPatchServices
 		{
 			if (!ChunkFile.bIsRetired)
 			{
-				AllChunkDbSize += ChunkFile.Archive->TotalSize();
+				AllChunkDbSize += ChunkFile.TotalSize();
 			}
 		}
 		
@@ -282,7 +381,7 @@ namespace BuildPatchServices
 				if (!ChunkFile.bIsRetired &&
 					ChunkFile.RetireAt < RetireAtEquivalent)
 				{
-					TotalSizeAtIndex += ChunkFile.Archive->TotalSize();
+					TotalSizeAtIndex += ChunkFile.TotalSize();
 				}
 			}
 
@@ -302,7 +401,7 @@ namespace BuildPatchServices
 		TArray<FChunkDbDataAccess> ChunkFiles;
 		TMap<FGuid, FChunkAccessLookup> ChunkGuidLookup;
 
-		FChunkDbChunkSource::LoadChunkDbFiles(ChunkDbFiles, FileSystem, ChunkAccessOrderedList, ChunkFiles, ChunkGuidLookup, nullptr);
+		FChunkDbChunkSource::LoadChunkDbFiles(ChunkDbFiles, FileSystem, ChunkAccessOrderedList, ChunkFiles, ChunkGuidLookup, nullptr, false);
 
 		return GetChunkDbSizesAtIndexesInternal(ChunkFiles, ChunkAccessOrderedList.Num(), FileCompletionIndexes, OutChunkDbSizesAtCompletion);
 	}
@@ -318,7 +417,7 @@ namespace BuildPatchServices
 
 		OriginalChunkCount = InChunkAccessOrderList.Num();
 
-		LoadChunkDbFiles(Configuration.ChunkDbFiles, FileSystem, InChunkAccessOrderList, ChunkDbDataAccesses, ChunkDbDataAccessLookup, &AvailableChunks);
+		LoadChunkDbFiles(Configuration.ChunkDbFiles, FileSystem, InChunkAccessOrderList, ChunkDbDataAccesses, ChunkDbDataAccessLookup, &AvailableChunks, Configuration.bTruncateChunkDbs);
 
 		// Immediately retire any chunkdbs we don't need so they don't eat disk space during the first file.
 		FChunkDbChunkSource::ReportFileCompletion(OriginalChunkCount);
@@ -334,7 +433,7 @@ namespace BuildPatchServices
 			return [](bool) {return;};
 		}
 
-		return [ChunkInfo, DataId, DestinationBuffer, UserPtr, CompleteFn, ChunkDataSerialization = ChunkDataSerialization, ChunkDbChunkSourceStat = ChunkDbChunkSourceStat](bool bIsAborted)
+		return [ChunkInfo, DataId, DestinationBuffer, UserPtr, CompleteFn, ChunkDataSerialization = ChunkDataSerialization, ChunkDbChunkSourceStat = ChunkDbChunkSourceStat, bTruncateChunkDbs = Configuration.bTruncateChunkDbs](bool bIsAborted)
 		{
 			TRACE_CPUPROFILER_EVENT_SCOPE(ChunkDbRead);
 			if (bIsAborted)
@@ -342,33 +441,42 @@ namespace BuildPatchServices
 				CompleteFn.Execute(DataId, true, false, UserPtr);
 				return;
 			}
-
 			FChunkLocation& ChunkLocation = *ChunkInfo->Location;
+			ISpeedRecorder::FRecord ActivityRecord;
+			bool Result = false;
+			FUniqueBuffer CompressedBuffer;
+			IChunkDbChunkSourceStat::ELoadResult LoadResult = IChunkDbChunkSourceStat::ELoadResult::Success;
+			FChunkHeader Header;
 			FChunkDbDataAccess& ChunkDbDataAccess = *ChunkInfo->DbFile;
-			FArchive* ChunkDbFile = ChunkDbDataAccess.Archive.Get();
-			if (ChunkDbFile->IsError())
 			{
-				CompleteFn.Execute(DataId, false, true, UserPtr);
-				return;
-			}
-
-			IChunkDbChunkSourceStat::ELoadResult LoadResult = IChunkDbChunkSourceStat::ELoadResult::Success;
-			ISpeedRecorder::FRecord ActivityRecord;
-			ActivityRecord.CyclesStart = FStatsCollector::GetCycles();
-			ChunkDbChunkSourceStat->OnLoadStarted(DataId);
+				FArchive* ChunkDbFile;
+				TUniquePtr<UE::TScopeLock<UE::FMutex>> LockGuard = ChunkDbDataAccess.LockArchive(&ChunkDbFile);
+				if (!ChunkDbFile || ChunkDbFile->IsError())
+				{
+					CompleteFn.Execute(DataId, false, true, UserPtr);
+					return;
+				}
+				
+				ActivityRecord.CyclesStart = FStatsCollector::GetCycles();
+				ChunkDbChunkSourceStat->OnLoadStarted(DataId);
 
-			// We'd love to read direct in to the destination if we don't have
-			// any compression. However we don't know if it's compressed until we read the header which
-			// is tiny and dependent - we don't know how big it is until we read part of it.
-			ChunkDbFile->Seek(ChunkLocation.ByteStart);
-			FChunkHeader Header;
+				if (static_cast<uint64>(ChunkDbFile->TotalSize()) < ChunkLocation.ByteStart)
+				{
+					GLog->Logf(ELogVerbosity::Type::Error, TEXT("Total file size (%lld) < Chunk start position (%lld). The file may have been truncated before resuming unlock: %s"), ChunkDbFile->TotalSize(), ChunkLocation.ByteStart, *ChunkDbDataAccess.ChunkDbFileName);
+					CompleteFn.Execute(DataId, false, true, UserPtr);
+					return;
+				}
+				// We'd love to read direct in to the destination if we don't have
+				// any compression. However we don't know if it's compressed until we read the header which
+				// is tiny and dependent - we don't know how big it is until we read part of it.
+				ChunkDbFile->Seek(ChunkLocation.ByteStart);
 
-			// If it's uncompressed, we can read direct to the destination.
-			FUniqueBuffer CompressedBuffer;
-			bool Result = ChunkDataSerialization->ValidateAndRead(*ChunkDbFile, DestinationBuffer, Header, CompressedBuffer);
+				// If it's uncompressed, we can read direct to the destination.
+				Result = ChunkDataSerialization->ValidateAndRead(*ChunkDbFile, DestinationBuffer, Header, CompressedBuffer);
 
-			// Save this here so we only include the IO time and not the hash/decompress time.
-			ActivityRecord.Size = ChunkDbFile->Tell() - ChunkLocation.ByteStart;
+				// Save this here so we only include the IO time and not the hash/decompress time.
+				ActivityRecord.Size = ChunkDbFile->Tell() - ChunkLocation.ByteStart;
+			}
 			ActivityRecord.CyclesEnd = FStatsCollector::GetCycles();
 			ChunkDbChunkSourceStat->OnReadComplete(ActivityRecord);
 
@@ -389,7 +497,7 @@ namespace BuildPatchServices
 				[DataId, CompleteFn, UserPtr, Header, DestinationBuffer, ChunkDataSerialization = ChunkDataSerialization, 
 				CompressedBuffer = MoveTemp(CompressedBuffer), ChunkDbChunkSourceStat = ChunkDbChunkSourceStat]()
 				{
-					bool bDecompressSucceeded = ChunkDataSerialization->DecompressValidatedRead(Header, DestinationBuffer, CompressedBuffer);
+					bool bDecompressSucceeded = ChunkDataSerialization->DecompressAndDecryptValidatedRead(Header, DestinationBuffer, CompressedBuffer);
 
 					ChunkDbChunkSourceStat->OnLoadComplete(DataId, 
 						bDecompressSucceeded ? IChunkDbChunkSourceStat::ELoadResult::Success : IChunkDbChunkSourceStat::ELoadResult::CorruptedData);
diff --git a/Engine/Source/Runtime/Online/BuildPatchServices/Private/Installer/ChunkDbChunkSource.h b/Engine/Source/Runtime/Online/BuildPatchServices/Private/Installer/ChunkDbChunkSource.h
--- a/Engine/Source/Runtime/Online/BuildPatchServices/Private/Installer/ChunkDbChunkSource.h
+++ b/Engine/Source/Runtime/Online/BuildPatchServices/Private/Installer/ChunkDbChunkSource.h
@@ -28,12 +28,16 @@ namespace BuildPatchServices
 		// If true, once we complete a file we delete all the chunkdbs used to create it.
 		bool bDeleteChunkDBAfterUse = false;
 
+		// If we are loading from a chunkDB, in reverse order, we can truncate the chunkDB files as we go if desired.
+		// This is for preload support.
+		bool bTruncateChunkDbs = false;
+
 		/**
 		* Constructor which sets usual defaults, and takes params for values that cannot use a default.
 		* @param InChunkDbFiles    The chunkdb filename array.
 		*/
-		FChunkDbSourceConfig(const TArray<FString>& InChunkDbFiles)
-		: ChunkDbFiles(InChunkDbFiles)
+		FChunkDbSourceConfig(TArray<FString> InChunkDbFiles)
+			: ChunkDbFiles(MoveTemp(InChunkDbFiles))
 		{}
 	};
 
@@ -57,8 +61,9 @@ namespace BuildPatchServices
 
 		// As above, except use the remaining open chunkdbs for progressive disk space checking.		
 		virtual uint64 GetChunkDbSizesAtIndexes(const TArray<int32>& FileCompletionIndexes, TArray<uint64>& OutChunkDbSizesAtCompletion) const = 0;
-
+		
 		virtual void ReportFileCompletion(int32 RemainingChunkCount) =0;
+		virtual void TryTruncateFiles(int32 RemainingChunkCount) =0;
 
 		static IConstructorChunkDbChunkSource* CreateChunkDbSource(FChunkDbSourceConfig&& Configuration, IFileSystem* FileSystem, const TArray<FGuid>& ChunkAccessOrderList, 
 			IChunkDataSerialization* ChunkDataSerialization, IChunkDbChunkSourceStat* ChunkDbChunkSourceStat);
diff --git a/Engine/Source/Runtime/Online/BuildPatchServices/Private/Installer/ChunkReferenceTracker.h b/Engine/Source/Runtime/Online/BuildPatchServices/Private/Installer/ChunkReferenceTracker.h
--- a/Engine/Source/Runtime/Online/BuildPatchServices/Private/Installer/ChunkReferenceTracker.h
+++ b/Engine/Source/Runtime/Online/BuildPatchServices/Private/Installer/ChunkReferenceTracker.h
@@ -26,6 +26,8 @@ namespace BuildPatchServices
 		/**
 		 * Gets a set of all chunks referenced by the installation this tracker refers to.
 		 * This returns the _remaining_ referenced chunks!
+		 * Note that this does NOT return chunks in usage order anymore! If you were tacitly relying
+		 * on the enumeration order of an unordered TSet, this likely broke! Use CopyOutOrderedUseList.
 		 * @return set of GUID containing every chunk used.
 		 */
 		virtual TSet<FGuid> GetReferencedChunks() const = 0;
@@ -247,10 +249,11 @@ namespace BuildPatchServices
 			TArray<FString> TaggedFiles;
 			TSet<FGuid> TaggedNewChunks;
 			InstallManifest->GetTaggedFileList(InstallTagSet, TaggedFiles);
-			for (const FString& File : TaggedFiles)
+			for (const FString& TaggedFile : TaggedFiles)
 			{
-				const FFileManifest* NewFileManifest = InstallManifest->GetFileManifest(File);
-				if (NewFileManifest != nullptr)
+				const FFileManifest* const OldFileManifest = CurrentManifest->GetFileManifest(TaggedFile);
+				const FFileManifest* const NewFileManifest = InstallManifest->GetFileManifest(TaggedFile);
+				if (!OldFileManifest || OldFileManifest->SHA1Hash != NewFileManifest->SHA1Hash || !FilesInstalled.Contains(TaggedFile))
 				{
 					for (const FChunkPart& ChunkPart : NewFileManifest->ChunkParts)
 					{
diff --git a/Engine/Source/Runtime/Online/BuildPatchServices/Private/Installer/CloudChunkSource.cpp b/Engine/Source/Runtime/Online/BuildPatchServices/Private/Installer/CloudChunkSource.cpp
--- a/Engine/Source/Runtime/Online/BuildPatchServices/Private/Installer/CloudChunkSource.cpp
+++ b/Engine/Source/Runtime/Online/BuildPatchServices/Private/Installer/CloudChunkSource.cpp
@@ -5,10 +5,12 @@
 #include "HAL/ThreadSafeBool.h"
 #include "Containers/Queue.h"
 #include "Misc/ScopeLock.h"
+#include "Algo/AllOf.h"
 #include "Async/Async.h"
 #include "Async/Future.h"
 #include "Core/MeanValue.h"
 #include "Core/Platform.h"
+#include "Http.h"
 #include "Installer/ChunkReferenceTracker.h"
 #include "Installer/ChunkStore.h"
 #include "Installer/DownloadService.h"
@@ -23,7 +25,7 @@
 #include "Memory/SharedBuffer.h"
 #include "Tasks/Task.h"
 #include "ProfilingDebugging/CountersTrace.h"
-
+#include <atomic>
 DECLARE_LOG_CATEGORY_EXTERN(LogCloudChunkSource, Log, All);
 DEFINE_LOG_CATEGORY(LogCloudChunkSource);
 
@@ -126,6 +128,150 @@ namespace BuildPatchServices
 	}
 
 
+	/**
+	 * Helper class for limiting download rate using an implementation of the token bucket algorithm.
+	 * 
+	 * One 'token' in this context is a byte that we can schedule for download. On tick, we increase
+	 * the number of available tokens according to the configured download rate limit. When we want
+	 * to schedule a chunk for download we must confirm that enough tokens are available and consume
+	 * them appropriately. The result is that the average download rate cannot exceed the configured
+	 * limit, although short bursts of higher rates are likely.
+	 */
+	class FDownloadRateLimiter
+	{
+	public:
+		FDownloadRateLimiter(uint64 InRateLimitBps)
+			: RateLimitBps(InRateLimitBps)
+			, TokenCount(InRateLimitBps)
+			, LastTickCycleCount(0) {}
+
+		/**
+		 * Tick this component. Must be called regularly from a single thread.
+		 */
+		void Tick()
+		{
+			// Sample the max burst size once now since it could be
+			// changed by another thread calling SetRateLimitBps.
+			const uint64 MaxBurstValue = MaxBurstSize();
+
+			uint64 CurrentTickCycleCount = FStatsCollector::GetCycles();
+
+			// On the first tick the only thing we should do is update LastTickCycleCount.
+			if (LastTickCycleCount != 0)
+			{
+				// Count the time elapsed since last tick and increment TokenCount appropriately.
+				uint64 ExpectedTokenCount, NewTokenCount;
+				do
+				{
+					ExpectedTokenCount = TokenCount.Load();
+					CurrentTickCycleCount = FStatsCollector::GetCycles();
+					const uint64 Increment = (uint64)(RateLimitBps.Load() * FStatsCollector::CyclesToSeconds(CurrentTickCycleCount - LastTickCycleCount));
+					NewTokenCount = FMath::Min(ExpectedTokenCount + Increment, MaxBurstValue);
+
+					// If the compare exchange fails because another thread
+					// changed the value in between, try again.
+				}
+				while (!TokenCount.CompareExchange(ExpectedTokenCount, NewTokenCount));
+			}
+			LastTickCycleCount = CurrentTickCycleCount;
+		}
+
+		/**
+		 * Update the target rate limit.
+		 * @param InRateLimitBps Rate limit in bytes per second. 0 means unlimited.
+		 */
+		void SetRateLimitBps(uint64 InRateLimitBps)
+		{
+			RateLimitBps.Exchange(InRateLimitBps);
+			TokenCount.Exchange(InRateLimitBps);
+		}
+
+		/**
+		 * Attempt to schedule a chunk download of ChunkSizeBytes.
+		 * @param ChunkSizeBytes Size of the chunk.
+		 * Returns 0 if the chunk can be scheduled, or the estimated ms until it can.
+		 */
+		uint64 TrySchedule(uint64 ChunkSizeBytes)
+		{
+			// Sample the max burst size once now since it could be
+			// changed by another thread calling SetRateLimitBps.
+			const uint64 MaxBurstValue = MaxBurstSize();
+
+			// Compute the token cost for this chunk. There are two possible cases:
+			// 1. ChunkSizeBytes is <= maximum burst size.
+			//    The cost is simply equal to ChunkSizeBytes.
+			// 2. ChunkSizeBytes is > maximum burst size.
+			//    The cost is equal to the maximum burst size. This avoids deadlocks
+			//    regardless of the relationship between the rate limit and chunk size,
+			//    even if it means that the rate limit cannot be strictly adhered to.
+			const uint64 TokenCost = ChunkSizeBytes > MaxBurstValue ? MaxBurstValue : ChunkSizeBytes;
+
+			//UE_LOG(LogTemp, Warning, TEXT("Scheduling: %llu bytes, cost %llu, tokens %llu"), ChunkSizeBytes, TokenCost, TokenCount.Load());
+			uint64 BytesPerSecondLimit = RateLimitBps.Load();
+
+			if (BytesPerSecondLimit != 0)
+			{
+				// Do an atomic compare exchange to make sure we don't double decrement
+				// and end up with e.g. negative TokenCount.
+				uint64 ExpectedTokenCount, NewTokenCount;
+				do
+				{
+					// If the cost is greater than what's available, return false.
+					ExpectedTokenCount = TokenCount.Load();
+					if (ExpectedTokenCount < TokenCost)
+					{
+						// We know the rate and how many we are missing, estimate how long until it's ready.
+						uint64 NeededBytes = TokenCost - ExpectedTokenCount;
+						uint64 BytesPerMsLimit = FMath::Max(1ULL, BytesPerSecondLimit / 1000);
+
+						// 0 means it can be scheduled, return a min of 1 ms delay
+						uint64 MsUntilToken = FMath::Max(1ULL, NeededBytes / BytesPerMsLimit);
+						return MsUntilToken;
+					}
+
+					// Otherwise atomically decrement and return true.
+					NewTokenCount = ExpectedTokenCount - TokenCost;
+
+					// If the compare exchange fails because another thread
+					// changed the value in between, try again.
+				}
+				while (!TokenCount.CompareExchange(ExpectedTokenCount, NewTokenCount));
+			}
+			return 0;
+		}
+
+		/**
+		 * Refunds token based on the ChunkSizeBytes.
+		 * @param ChunkSizeBytes Size of the chunk.
+		 */
+		void Refund(uint64 ChunkSizeBytes)
+		{
+			if (RateLimitBps.Load() != 0)
+			{
+				uint64 ExpectedTokenCount, NewTokenCount;
+				do
+				{
+					ExpectedTokenCount = TokenCount.Load();
+					NewTokenCount = FMath::Min(MaxBurstSize(), ExpectedTokenCount + ChunkSizeBytes);
+				}
+				while (!TokenCount.CompareExchange(ExpectedTokenCount, NewTokenCount));
+			}
+		}
+
+	private:
+		/**
+		 * Returns the current maximum burst size, which is also
+		 * the maximum number of tokens that may be stored.
+		 */
+		uint64 MaxBurstSize(void) const { return RateLimitBps.Load() * BurstMultiplier; }
+
+	private:
+		static const uint64 BurstMultiplier = 4;
+		TAtomic<uint64> RateLimitBps;
+		TAtomic<uint64> TokenCount;
+		uint64 LastTickCycleCount;
+	};
+	
 	/**
 	 * The concrete implementation of ICloudChunkSource
 	 */
@@ -134,20 +280,27 @@ namespace BuildPatchServices
 	{
 	private:
 		/**
-		 * This is a wrapper class for binding thread safe shared ptr delegates for the download service, without having to enforce that
+		 * This is a mediator class for binding thread safe shared ptr delegates for the download service, without having to enforce that
 		 * this service should be made using TShared* reference controllers.
 		 */
 		class FDownloadDelegates
 		{
-		public:
-			FDownloadDelegates(FCloudChunkSource& InCloudChunkSource);
-
 		public:
 			void OnDownloadProgress(int32 RequestId, uint64 BytesSoFar);
 			void OnDownloadComplete(int32 RequestId, const FDownloadRef& Download);
 
+			TMap<int32, FDownloadRef> GetCompletedDownloads();
+
+			// Returns the number of seconds since the last data was received.
+			// Tracking health and connection state
+			double GetSecondsSinceLastData() const;
 		private:
-			FCloudChunkSource& CloudChunkSource;
+			// The point in time when the latest data was received;
+			std::atomic<int64> CyclesAtLastData{0};
+
+			// Communication from download thread to processing thread.
+			FCriticalSection CompletedDownloadsCS;
+			TMap<int32, FDownloadRef> CompletedDownloads;
 		};
 		
 		/**
@@ -176,6 +329,7 @@ namespace BuildPatchServices
 		// IControllable interface begin.
 		virtual void SetPaused(bool bInIsPaused) override;
 		virtual void Abort() override;
+		virtual void Reset() override;
 		// IControllable interface end.
 
 		// IChunkSource interface begin.
@@ -186,6 +340,7 @@ namespace BuildPatchServices
 		// IChunkSource interface end.
 
 		// ICloudChunkSource interface begin.
+		virtual void SetDownloadBytesPerSecondLimit(uint64 InDownloadRateLimitBps) override;
 		virtual void ThreadRun() override;
 		// ICloudChunkSource interface end.
 
@@ -193,11 +348,8 @@ namespace BuildPatchServices
 		void EnsureAquiring(const FGuid& DataId);
 
 		FGuid GetNextTask(const TMap<FGuid, FTaskInfo>& TaskInfos, const TMap<int32, FGuid>& InFlightDownloads, const TSet<FGuid>& TotalRequiredChunks, const TSet<FGuid>& PriorityRequests, const TSet<FGuid>& FailedDownloads, const TSet<FGuid>& Stored, TArray<FGuid>& DownloadQueue, EBuildPatchDownloadHealth DownloadHealth);
-		void OnDownloadProgress(int32 RequestId, uint64 BytesSoFar);
-		void OnDownloadComplete(int32 RequestId, const FDownloadRef& Download);
 
 	private:
-		TSharedRef<FDownloadDelegates, ESPMode::ThreadSafe> DownloadDelegates;
 		const FCloudSourceConfig Configuration;
 		IPlatform* Platform;
 		IChunkStore* ChunkStore;
@@ -212,19 +364,13 @@ namespace BuildPatchServices
 		TPromise<void> Promise;
 		TFuture<void> Future;
 		IBuildInstallerThread* Thread = nullptr;
-		FDownloadProgressDelegate OnDownloadProgressDelegate;
-		FDownloadCompleteDelegate OnDownloadCompleteDelegate;
-
-		// Tracking health and connection state
-		volatile int64 CyclesAtLastData;
 
 		// Communication from external process requesting pause/abort.
 		FThreadSafeBool bIsPaused;
 		FThreadSafeBool bShouldAbort;
 
-		// Communication from download thread to processing thread.
-		FCriticalSection CompletedDownloadsCS;
-		TMap<int32, FDownloadRef> CompletedDownloads;
+		// Download throttling helper.
+		FDownloadRateLimiter DownloadRateLimiter;
 
 		// Communication from request threads to processing thread.
 		FCriticalSection RequestedDownloadsCS;
@@ -246,19 +392,26 @@ namespace BuildPatchServices
 		TArray<int32> CloudDirFailureCount;
 	};
 
-	FCloudChunkSource::FDownloadDelegates::FDownloadDelegates(FCloudChunkSource& InCloudChunkSource)
-		: CloudChunkSource(InCloudChunkSource)
+	void FCloudChunkSource::FDownloadDelegates::OnDownloadProgress(int32 RequestId, uint64 BytesSoFar)
 	{
+		CyclesAtLastData = FStatsCollector::GetCycles();
 	}
 
-	void FCloudChunkSource::FDownloadDelegates::OnDownloadProgress(int32 RequestId, uint64 BytesSoFar)
+	void FCloudChunkSource::FDownloadDelegates::OnDownloadComplete(int32 RequestId, const FDownloadRef& Download)
 	{
-		CloudChunkSource.OnDownloadProgress(RequestId, BytesSoFar);
+		FScopeLock ScopeLock(&CompletedDownloadsCS);
+		CompletedDownloads.Add(RequestId, Download);
 	}
 
-	void FCloudChunkSource::FDownloadDelegates::OnDownloadComplete(int32 RequestId, const FDownloadRef& Download)
+	double FCloudChunkSource::FDownloadDelegates::GetSecondsSinceLastData() const
 	{
-		CloudChunkSource.OnDownloadComplete(RequestId, Download);
+		return FStatsCollector::CyclesToSeconds(FStatsCollector::GetCycles() - CyclesAtLastData);
+	}
+
+	TMap<int32, FDownloadRef> FCloudChunkSource::FDownloadDelegates::GetCompletedDownloads()
+	{
+		FScopeLock ScopeLock(&CompletedDownloadsCS);
+		return MoveTemp(CompletedDownloads);
 	}
 
 	FCloudChunkSource::FTaskInfo::FTaskInfo()
@@ -270,8 +423,7 @@ namespace BuildPatchServices
 	}
 
 	FCloudChunkSource::FCloudChunkSource(FCloudSourceConfig InConfiguration, IPlatform* InPlatform, IChunkStore* InChunkStore, IDownloadService* InDownloadService, IChunkReferenceTracker* InChunkReferenceTracker, IChunkDataSerialization* InChunkDataSerialization, IMessagePump* InMessagePump, IInstallerError* InInstallerError, IDownloadConnectionCount* InDownloadConnectionCount, ICloudChunkSourceStat* InCloudChunkSourceStat, IBuildManifestSet* InManifestSet, TSet<FGuid> InInitialDownloadSet)
-		: DownloadDelegates(MakeShareable(new FDownloadDelegates(*this)))
-		, Configuration(MoveTemp(InConfiguration))
+		: Configuration(MoveTemp(InConfiguration))
 		, Platform(InPlatform)
 		, ChunkStore(InChunkStore)
 		, DownloadService(InDownloadService)
@@ -285,13 +437,9 @@ namespace BuildPatchServices
 		, Promise()
 		, Future()
 		, Thread(nullptr)
-		, OnDownloadProgressDelegate(FDownloadProgressDelegate::CreateThreadSafeSP(DownloadDelegates, &FDownloadDelegates::OnDownloadProgress))
-		, OnDownloadCompleteDelegate(FDownloadCompleteDelegate::CreateThreadSafeSP(DownloadDelegates, &FDownloadDelegates::OnDownloadComplete))
-		, CyclesAtLastData(0)
 		, bIsPaused(false)
 		, bShouldAbort(false)
-		, CompletedDownloadsCS()
-		, CompletedDownloads()
+		, DownloadRateLimiter(InConfiguration.DownloadRateLimitBps)
 		, RequestedDownloadsCS()
 		, RequestedDownloads()
 		, DownloadCount(InDownloadConnectionCount)
@@ -329,6 +477,11 @@ namespace BuildPatchServices
 		bShouldAbort = true;
 	}
 
+	void FCloudChunkSource::Reset()
+	{
+		bShouldAbort = false;
+	}
+
 	IChunkDataAccess* FCloudChunkSource::Get(const FGuid& DataId)
 	{
 		IChunkDataAccess* ChunkData = ChunkStore->Get(DataId);
@@ -366,6 +519,11 @@ namespace BuildPatchServices
 		// We don't have a concept of being unavailable yet.
 	}
 
+	void FCloudChunkSource::SetDownloadBytesPerSecondLimit(uint64 InDownloadRateLimitBps)
+	{
+		DownloadRateLimiter.SetRateLimitBps(InDownloadRateLimitBps);
+	}
+
 	void FCloudChunkSource::EnsureAquiring(const FGuid& DataId)
 	{
 		FScopeLock ScopeLock(&RequestedDownloadsCS);
@@ -434,7 +592,7 @@ namespace BuildPatchServices
 			// Return the next chunk in the queue
 			if (DownloadQueue.Num() > 0)
 			{
-				return DownloadQueue.Pop(EAllowShrinking::No);
+				return DownloadQueue.Last();
 			}
 		}
 
@@ -455,15 +613,24 @@ namespace BuildPatchServices
 		FChunkSuccessRate ChunkSuccessRate;
 		EBuildPatchDownloadHealth TrackedDownloadHealth = EBuildPatchDownloadHealth::Excellent;
 		int32 TrackedActiveRequestCount = 0;
+		int32 TrackedActiveRequestAverageSum = 0;
+		int32 TrackedActiveRequestAverageCount = 0;
+		double SecondsAtLastTrackedAverage = 0;
 		TSet<FGuid> TotalRequiredChunks = InitialDownloadSet;
 		uint64 TotalRequiredChunkSize = ManifestSet->GetDownloadSize(TotalRequiredChunks);
 		uint64 TotalReceivedData = 0;
 		uint64 RepeatRequirementSize = 0;
+		const TSharedRef<FDownloadDelegates, ESPMode::ThreadSafe> DownloadDelegates = MakeShared<FDownloadDelegates, ESPMode::ThreadSafe>();
+
+		FDownloadProgressDelegate OnDownloadProgressDelegate(FDownloadProgressDelegate::CreateThreadSafeSP(DownloadDelegates, &FDownloadDelegates::OnDownloadProgress));
+		FDownloadCompleteDelegate OnDownloadCompleteDelegate(FDownloadCompleteDelegate::CreateThreadSafeSP(DownloadDelegates, &FDownloadDelegates::OnDownloadComplete));
 
 		// Chunk Uri Processing
 		typedef TTuple<FGuid, FChunkUriResponse> FGuidUriResponse;
-		typedef TQueue<FGuidUriResponse, EQueueMode::Mpsc> FGuidUriResponseQueue; // use Mpsc, message pump callback may be on this thread or message pump thread
-		TSharedRef<FGuidUriResponseQueue> ChunkUriResponsesRef = MakeShared<FGuidUriResponseQueue>();
+		typedef TQueue<FGuidUriResponse, EQueueMode::Mpsc> FGuidUriResponseQueue;
+		TSharedRef<FGuidUriResponseQueue, ESPMode::ThreadSafe> ChunkUriResponsesRef = MakeShared<FGuidUriResponseQueue>(); // use Mpsc, message pump callback may be on this thread or message pump thread
+		TWeakPtr<FGuidUriResponseQueue, ESPMode::ThreadSafe> WeakChunkUriResponses = ChunkUriResponsesRef;
+		FGuidUriResponseQueue& ChunkUriResponses = ChunkUriResponsesRef.Get();
 		TSet<FGuid> RequestedChunkUris;
 		TMap<FGuid, FChunkUriResponse> ChunkUris;
 
@@ -477,12 +644,15 @@ namespace BuildPatchServices
 		while (!bShouldAbort)
 		{
 			bool bRequiredDataUpdated = false;
+			DownloadRateLimiter.Tick();
+
 			// 'Forget' any repeat requirements.
 			FGuid RepeatRequirement;
 			while (RepeatRequirementMessages.Dequeue(RepeatRequirement))
 			{
 				if (PlacedInStore.Remove(RepeatRequirement) > 0)
 				{
+					RequestedChunkUris.Remove(RepeatRequirement);
 					RepeatRequirementSize += ManifestSet->GetDownloadSize(RepeatRequirement);
 					bRequiredDataUpdated = true;
 				}
@@ -500,20 +670,50 @@ namespace BuildPatchServices
 					bRequiredDataUpdated = true;
 				}
 			}
+			// Grab incoming requests as a priority.
+			TArray<FGuid> FrameRequestedDownloads;
+			RequestedDownloadsCS.Lock();
+			FrameRequestedDownloads = MoveTemp(RequestedDownloads);
+			RequestedDownloadsCS.Unlock();
+			TSet<FGuid> PriorityUrlRequests;
+			for (const FGuid& FrameRequestedDownload : FrameRequestedDownloads)
+			{
+				bDownloadsStarted = true;
+				if (!TaskInfos.Contains(FrameRequestedDownload) && !PlacedInStore.Contains(FrameRequestedDownload))
+				{
+					PriorityRequests.Add(FrameRequestedDownload);
+					if (!RequestedChunkUris.Contains(FrameRequestedDownload))
+					{
+						PriorityUrlRequests.Add(FrameRequestedDownload);
+						UE_LOG(LogCloudChunkSource, Log, TEXT("New chunk uri priority request: %s"), *FrameRequestedDownload.ToString());
+					}
+					else
+					{
+						UE_LOG(LogCloudChunkSource, Log, TEXT("New priority request: %s"), *FrameRequestedDownload.ToString());
+					}
+					if (!TotalRequiredChunks.Contains(FrameRequestedDownload))
+					{
+						TotalRequiredChunks.Add(FrameRequestedDownload);
+						TotalRequiredChunkSize += ManifestSet->GetDownloadSize(FrameRequestedDownload);
+						bRequiredDataUpdated = true;
+					}
+				}
+			}
 			// Select the next X chunks that are for downloading, so we can request URIs.
-			TFunction<bool(const FGuid&)> SelectPredicate = [&TotalRequiredChunks, &RequestedChunkUris](const FGuid& ChunkId) 
-			{ 
+			TFunction<bool(const FGuid&)> SelectPredicate = [&TotalRequiredChunks, &RequestedChunkUris](const FGuid& ChunkId)
+			{
 				// if we requre it and we haven't already requested it.
-				return TotalRequiredChunks.Contains(ChunkId) && !RequestedChunkUris.Contains(ChunkId); 
+				return TotalRequiredChunks.Contains(ChunkId) && !RequestedChunkUris.Contains(ChunkId);
 			};
-			TArray<FGuid> ChunkUrisToRequest;
 
+			TSet<FGuid> ChunkUrisToRequest;
 			// Don't take the lock over the reference stack if we can't ever pass our selection predicate
 			if (TotalRequiredChunks.Num())
 			{
-				ChunkUrisToRequest = ChunkReferenceTracker->SelectFromNextReferences(Configuration.PreFetchMaximum, SelectPredicate);
+				ChunkUrisToRequest = TSet<FGuid> { ChunkReferenceTracker->SelectFromNextReferences(FGenericPlatformMath::Max(1, Configuration.PreFetchMaximum - PriorityUrlRequests.Num()), SelectPredicate) };
 			}
-			
+			ChunkUrisToRequest.Append(MoveTemp(PriorityUrlRequests));
+
 			for (const FGuid& ChunkUriToRequest : ChunkUrisToRequest)
 			{
 				RequestedChunkUris.Add(ChunkUriToRequest);
@@ -522,25 +722,35 @@ namespace BuildPatchServices
 				FTaskInfo& Info = TaskInfos.FindOrAdd(ChunkUriToRequest);
 				Info.CloudDirUsed = CurrentBestCloudDir;
 				ChunkUriRequest.CloudDirectory = Configuration.CloudRoots[CurrentBestCloudDir];
-				ChunkUriRequest.RelativePath = ManifestSet->GetDataFilename(ChunkUriToRequest);
-				ChunkUriRequest.RelativePath.RemoveFromStart(TEXT("/"));
+				ChunkUriRequest.ChunkCloudSubdirectory = ManifestSet->GetChunkCloudSubdirectory(ChunkUriToRequest);
+				ChunkUriRequest.ChunkCloudSubdirectory.RemoveFromStart(TEXT("/"));
+				ChunkUriRequest.ChunkPath = ManifestSet->GetChunkFilename(ChunkUriToRequest);
+				ChunkUriRequest.UniqueBuildId = ManifestSet->GetChunkUniqueBuildId(ChunkUriToRequest);
+				ChunkUriRequest.bChunkFromDelta = ManifestSet->GetChunkFromDeltaStatus(ChunkUriToRequest);
 
-				bool bMessageSent = MessagePump->SendRequest(ChunkUriRequest, [ChunkUriResponsesRef, ChunkUriToRequest](FChunkUriResponse Response)
+				bool bMessageSent = MessagePump->SendRequest(ChunkUriRequest, [WeakChunkUriResponses, ChunkUriToRequest](FChunkUriResponse Response)
 				{
-					ChunkUriResponsesRef->Enqueue(FGuidUriResponse(ChunkUriToRequest, MoveTemp(Response)));
+					TSharedPtr<FGuidUriResponseQueue, ESPMode::ThreadSafe> ChunkUriResponsesPtr = WeakChunkUriResponses.Pin();
+					if (ChunkUriResponsesPtr.IsValid())
+					{
+						ChunkUriResponsesPtr->Enqueue(FGuidUriResponse(ChunkUriToRequest, MoveTemp(Response)));
+					}
 				});
-
 				if (!bMessageSent)
 				{
-					// If no handler is registered SendRequest does nothing - make our own default response
-					FChunkUriResponse Response;
-					Response.Uri = ChunkUriRequest.CloudDirectory / ChunkUriRequest.RelativePath;
-					ChunkUriResponsesRef->Enqueue(FGuidUriResponse(ChunkUriToRequest, MoveTemp(Response)));
+					TSharedPtr<FGuidUriResponseQueue, ESPMode::ThreadSafe> ChunkUriResponsesPtr = WeakChunkUriResponses.Pin();
+					if (ChunkUriResponsesPtr.IsValid())
+					{
+						// If no handler is registered SendRequest does nothing - make our own default response
+						FChunkUriResponse Response;
+						Response.Uri = ChunkUriRequest.CloudDirectory / ChunkUriRequest.ChunkCloudSubdirectory / ChunkUriRequest.ChunkPath;
+						ChunkUriResponsesPtr->Enqueue(FGuidUriResponse(ChunkUriToRequest, MoveTemp(Response)));
+					}
 				}
 			}
 			// Process new chunk uri responses.
 			FGuidUriResponse ChunkUriResponse;
-			for (FGuidUriResponseQueue& ChunkUriResponses = ChunkUriResponsesRef.Get(); ChunkUriResponses.Dequeue(ChunkUriResponse);)
+			while (ChunkUriResponses.Dequeue(ChunkUriResponse))
 			{
 				if (ChunkUriResponse.Get<1>().bFailed)
 				{
@@ -554,25 +764,6 @@ namespace BuildPatchServices
 					ChunkUris.Add(MoveTemp(ChunkUriResponse.Get<0>()), MoveTemp(ChunkUriResponse.Get<1>()));
 				}
 			}
-			// Grab incoming requests as a priority.
-			TArray<FGuid> FrameRequestedDownloads;
-			RequestedDownloadsCS.Lock();
-			FrameRequestedDownloads = MoveTemp(RequestedDownloads);
-			RequestedDownloadsCS.Unlock();
-			for (const FGuid& FrameRequestedDownload : FrameRequestedDownloads)
-			{
-				bDownloadsStarted = true;
-				if (!TaskInfos.Contains(FrameRequestedDownload) && !PlacedInStore.Contains(FrameRequestedDownload))
-				{
-					PriorityRequests.Add(FrameRequestedDownload);
-					if (!TotalRequiredChunks.Contains(FrameRequestedDownload))
-					{
-						TotalRequiredChunks.Add(FrameRequestedDownload);
-						TotalRequiredChunkSize += ManifestSet->GetDownloadSize(FrameRequestedDownload);
-						bRequiredDataUpdated = true;
-					}
-				}
-			}
 			// Trim our initial download list on first begin.
 			if (!bTotalRequiredTrimmed && bDownloadsStarted)
 			{
@@ -592,10 +783,8 @@ namespace BuildPatchServices
 			}
 
 			// Process completed downloads.
-			TMap<int32, FDownloadRef> FrameCompletedDownloads;
-			CompletedDownloadsCS.Lock();
-			FrameCompletedDownloads = MoveTemp(CompletedDownloads);
-			CompletedDownloadsCS.Unlock();
+			const TMap<int32, FDownloadRef> FrameCompletedDownloads = DownloadDelegates->GetCompletedDownloads();
+			
 			for (const TPair<int32, FDownloadRef>& FrameCompletedDownload : FrameCompletedDownloads)
 			{
 				const int32& RequestId = FrameCompletedDownload.Key;
@@ -614,31 +803,50 @@ namespace BuildPatchServices
 					FSHAHash ChunkShaHash;
 					if (ManifestSet->GetChunkShaHash(DownloadId, ChunkShaHash))
 					{
-						ChunkDataSerialization->InjectShaToChunkData(DownloadedData, ChunkShaHash);
+						// Older chunks might not have a sha hash internally, but the manifest 
+						// is new and have it, so we can inject it.
+						bDownloadSuccess = ChunkDataSerialization->InjectShaToChunkData(DownloadedData, ChunkShaHash);
+						if (!bDownloadSuccess)
+						{
+							CloudChunkSourceStat->OnDownloadCorrupt(DownloadId, TaskInfo.UrlUsed, EChunkLoadResult::SerializationError);
+							UE_LOG(LogCloudChunkSource, Log, TEXT("Chunk download request succeeded but the data header was not valid to inject ShaHash. DownloadedData size: %d for chunk: %s"), DownloadedData.Num(), *TaskInfo.UrlUsed);
+						}
 					}
 
-					EChunkLoadResult LoadResult;
-					TUniquePtr<IChunkDataAccess> ChunkDataAccess(ChunkDataSerialization->LoadFromMemory(DownloadedData, LoadResult));
-					bDownloadSuccess = LoadResult == EChunkLoadResult::Success;
 					if (bDownloadSuccess)
 					{
-						TotalReceivedData += TaskInfo.ExpectedSize;
-						TaskInfos.Remove(DownloadId);
-						PlacedInStore.Add(DownloadId);
-						ChunkStore->Put(DownloadId, MoveTemp(ChunkDataAccess));
-						CloudChunkSourceStat->OnDownloadSuccess(DownloadId);
-						CloudChunkSourceStat->OnReceivedDataUpdated(TotalReceivedData);
+						EChunkLoadResult LoadResult;
+						TUniquePtr<IChunkDataAccess> ChunkDataAccess(ChunkDataSerialization->LoadFromMemory(DownloadedData, LoadResult));
+						if (LoadResult == EChunkLoadResult::Success || LoadResult == EChunkLoadResult::MissingSecretKey)
+						{
+							TotalReceivedData += TaskInfo.ExpectedSize;
+							TaskInfos.Remove(DownloadId);
+							PlacedInStore.Add(DownloadId);
+							ChunkStore->Put(DownloadId, MoveTemp(ChunkDataAccess));
+							CloudChunkSourceStat->OnDownloadSuccess(DownloadId);
+							CloudChunkSourceStat->OnReceivedDataUpdated(TotalReceivedData);
+							CloudChunkSourceStat->OnCloudDirectoryDataReceived(Configuration.CloudRoots[TaskInfo.CloudDirUsed], TaskInfo.ExpectedSize);
+						}
+						else
+						{
+							CloudChunkSourceStat->OnDownloadCorrupt(DownloadId, TaskInfo.UrlUsed, LoadResult);
+							bDownloadSuccess = false;
+							UE_LOG(LogCloudChunkSource, Log, TEXT("Chunk download request succeeded but the data was not valid. Load result %s (%d) for chunk: %s"), ToString(LoadResult), LoadResult, *TaskInfo.UrlUsed);
+						}
 					}
-					else
+
+					// Return the verbosity back to the default value if it was changed
+					if (TaskInfo.RetryNum > Configuration.RetriesToEnableVerbose && UE_GET_LOG_VERBOSITY(LogHttp) != ELogVerbosity::Error)
 					{
-						CloudChunkSourceStat->OnDownloadCorrupt(DownloadId, TaskInfo.UrlUsed, LoadResult);
-						UE_LOG(LogCloudChunkSource, Error, TEXT("CORRUPT: %s"), *TaskInfo.UrlUsed);
+						UE_SET_LOG_VERBOSITY(LogHttp, Error);
 					}
 				}
 				else
 				{
 					CloudChunkSourceStat->OnDownloadFailed(DownloadId, TaskInfo.UrlUsed);
-					UE_LOG(LogCloudChunkSource, Error, TEXT("FAILED: %s"), *TaskInfo.UrlUsed);
+					FChunkSourceEvent ChunkDownloadFailEvent{ FChunkSourceEvent::EType::HttpRequestFailed, TaskInfo.UrlUsed };
+					MessagePump->SendMessage(ChunkDownloadFailEvent);
+					UE_LOG(LogCloudChunkSource, Log, TEXT("FAILED: %s"), *TaskInfo.UrlUsed);
 				}
 
 				// Handle failed (note this also launches a retry on a bad serialization, not just download.
@@ -652,6 +860,13 @@ namespace BuildPatchServices
 						bShouldAbort = true;
 					}
 					++TaskInfo.RetryNum;
+					
+					// Change logs verbosity to print the detailed error
+					if (TaskInfo.RetryNum > Configuration.RetriesToEnableVerbose)
+					{
+						UE_SET_LOG_VERBOSITY(LogHttp, Verbose);
+					}
+
 
 					// Mark this CDN as failed.
 					{
@@ -678,12 +893,18 @@ namespace BuildPatchServices
 
 						CurrentBestCloudDir = MinAtIndex;
 						UE_LOG(LogCloudChunkSource, Warning, TEXT("CDN %s failed download, updating CDN selection to: %s"), *Configuration.CloudRoots[TaskInfo.CloudDirUsed], *Configuration.CloudRoots[MinAtIndex]);
+						CloudChunkSourceStat->OnCloudDirectoryError(Configuration.CloudRoots[TaskInfo.CloudDirUsed]);
 					}
 
 					TaskInfo.SecondsAtFail = FStatsCollector::GetSeconds();
-
+					const int32 UnusedDownloadTokenSize = TaskInfo.ExpectedSize - Download->GetData().Num();
+					if (UnusedDownloadTokenSize > 0)
+					{
+						DownloadRateLimiter.Refund(UnusedDownloadTokenSize);
+					}
 					RequestedChunkUris.Remove(DownloadId);
-					ChunkUris.Remove(DownloadId);
+
+					UE_LOG(LogCloudChunkSource, Log, TEXT("Retry '%d' for chunk '%s'"), TaskInfo.RetryNum, *TaskInfo.UrlUsed);
 				}
 				else
 				{
@@ -691,6 +912,7 @@ namespace BuildPatchServices
 					MeanChunkTime.AddSample(ChunkTime);
 					ChunkSuccessRate.AddSuccess();
 				}
+				ChunkUris.Remove(DownloadId);
 				InFlightDownloads.Remove(RequestId);
 			}
 
@@ -701,9 +923,10 @@ namespace BuildPatchServices
 				if (TaskInfos.FindOrAdd(InFlightIt.Value()).RetryNum == 0)
 				{
 					bAllDownloadsRetrying = false;
+					break;
 				}
 			}
-			const double SecondsSinceData = FStatsCollector::CyclesToSeconds(FStatsCollector::GetCycles() - CyclesAtLastData);
+			const double SecondsSinceData = DownloadDelegates->GetSecondsSinceLastData();
 			const bool bReportAsDisconnected = (bAllDownloadsRetrying && SecondsSinceData > Configuration.DisconnectedDelay);
 			const float SuccessRate = ChunkSuccessRate.GetOverall();
 			EBuildPatchDownloadHealth OverallDownloadHealth = GetDownloadHealth(bReportAsDisconnected, Configuration.HealthPercentages, SuccessRate);
@@ -725,17 +948,20 @@ namespace BuildPatchServices
 				while ((NextTask = GetNextTask(TaskInfos, InFlightDownloads, TotalRequiredChunks, PriorityRequests, FailedDownloads, PlacedInStore, DownloadQueue, ImmediateDownloadHealth)).IsValid())
 				{
 					FChunkUriResponse* ChunkUri = ChunkUris.Find(NextTask);
-					if (ChunkUri)
+					uint64 DownloadSize = ManifestSet->GetDownloadSize(NextTask);
+					if (ChunkUri && DownloadRateLimiter.TrySchedule(DownloadSize) == 0)
 					{
 						FTaskInfo& TaskInfo = TaskInfos.FindOrAdd(NextTask);
 						TaskInfo.bQueuedForDownload = true;
 						TaskInfo.UrlUsed = ChunkUri->Uri;
-						TaskInfo.ExpectedSize = ManifestSet->GetDownloadSize(NextTask);
+						TaskInfo.ExpectedSize = DownloadSize;
 						TaskInfo.SecondsAtRequested = FStatsCollector::GetSeconds();
-						int32 RequestId = DownloadService->RequestFileWithHeaders(TaskInfo.UrlUsed, ChunkUri->AdditionalHeaders, OnDownloadCompleteDelegate, OnDownloadProgressDelegate);
+						int32 RequestId = 0;
+						DownloadService->RequestFileWithHeaders(TaskInfo.UrlUsed, ChunkUri->AdditionalHeaders, OnDownloadCompleteDelegate, OnDownloadProgressDelegate, &RequestId);
 						InFlightDownloads.Add(RequestId, NextTask);
 						PriorityRequests.Remove(NextTask);
 						FailedDownloads.Remove(NextTask);
+						DownloadQueue.Remove(NextTask);
 						CloudChunkSourceStat->OnDownloadRequested(NextTask);
 					}
 					else
@@ -753,6 +979,24 @@ namespace BuildPatchServices
 				CloudChunkSourceStat->OnActiveRequestCountUpdated(TrackedActiveRequestCount);
 			}
 
+			// Update Average request count.
+			const double CurrentSeconds = FStatsCollector::GetSeconds();
+			// Track the average once per second
+			if (SecondsAtLastTrackedAverage == 0 || CurrentSeconds - SecondsAtLastTrackedAverage > 1.0f)
+			{
+				SecondsAtLastTrackedAverage = CurrentSeconds;
+				// Start tracking the average only once we have started downloading
+				if (TrackedActiveRequestAverageSum > 0 || TrackedActiveRequestCount > 0)
+				{
+					TrackedActiveRequestAverageSum += TrackedActiveRequestCount;
+					TrackedActiveRequestAverageCount++;
+					if (TrackedActiveRequestAverageCount > 0)
+					{
+						CloudChunkSourceStat->OnActiveRequestCountAvgUpdated(TrackedActiveRequestAverageSum / TrackedActiveRequestAverageCount);
+					}
+				}
+			}
+
 			// Check for abnormally slow downloads. This was originally implemented as a temporary measure to fix major stall anomalies and zero size tcp window issue.
 			// It remains until proven unrequired.
 			if (MeanChunkTime.IsReliable())
@@ -810,17 +1054,6 @@ namespace BuildPatchServices
 		Promise.SetValue();
 	}
 
-	void FCloudChunkSource::OnDownloadProgress(int32 RequestId, uint64 BytesSoFar)
-	{
-		FPlatformAtomics::InterlockedExchange(&CyclesAtLastData, FStatsCollector::GetCycles());
-	}
-
-	void FCloudChunkSource::OnDownloadComplete(int32 RequestId, const FDownloadRef& Download)
-	{
-		FScopeLock ScopeLock(&CompletedDownloadsCS);
-		CompletedDownloads.Add(RequestId, Download);
-	}
-
 	
 	ICloudChunkSource* FCloudChunkSourceFactory::Create(FCloudSourceConfig Configuration, IPlatform* Platform, IChunkStore* ChunkStore, IDownloadService* DownloadService, IChunkReferenceTracker* ChunkReferenceTracker, IChunkDataSerialization* ChunkDataSerialization, IMessagePump* MessagePump, IInstallerError* InstallerError, IDownloadConnectionCount* ConnectionCount, ICloudChunkSourceStat* CloudChunkSourceStat, IBuildManifestSet* ManifestSet, TSet<FGuid> InitialDownloadSet)
 	{
@@ -883,13 +1116,14 @@ namespace BuildPatchServices
 			int32 CloudDirUsed = 0;
 			FString UrlUsed;
 			int32 RetryNum = 0;
-			int32 ExpectedSize = 0;
+			uint64 ExpectedSize = 0;
 			int32 RequestId = 0;
 			double SecondsAtRequested = 0;
 			int32 ReadId = 0;
 
 			// The time when we want to launch this retry, if RetryNum != 0
 			double RetryTime = 0;
+			bool bDecompressionSuccess = false; // Flag indicating that the chunk data was successfully downloaded, decoded, and decompressed.
 		};
 
 	public:
@@ -905,6 +1139,7 @@ namespace BuildPatchServices
 			, ManifestSet(InManifestSet)
 			, DownloadCount(InDownloadCount)
 			, MessagePump(InMessagePump)
+			, DownloadRateLimiter(InConfiguration.DownloadRateLimitBps)
 		{
 			CloudDirFailureCount.SetNumZeroed(Configuration.CloudRoots.Num());
 		}
@@ -959,7 +1194,40 @@ namespace BuildPatchServices
 			CloudChunkSourceStat->OnRequiredDataUpdated(InDownloadExpected);
 		}
 
-	private:	
+		virtual void SetDownloadBytesPerSecondLimit(uint64 BytesPerSecondLimit) override
+		{
+			DownloadRateLimiter.SetRateLimitBps(BytesPerSecondLimit);
+
+			// If the limit changed, we need to wake up and try to see if we can schedule
+			// more downloads
+			WakeupMainThreadFn();
+		}
+
+	private:
+
+		bool GetDisconnectionStatus()
+		{
+			{
+				FScopeLock _(&ReadQueueCS);
+
+				if (ReadQueue.Num() == 0 || !Algo::AllOf(ReadQueue, [](const FCloudRead* CloudRead) { return CloudRead->RetryNum > 0; }))
+				{
+					AllDownloadsRetryingSince.store(-1.0f);
+					return false;
+				}
+			}
+
+			// Store the timestamp when all chunks downloads are in retrying state
+			if (AllDownloadsRetryingSince < 0.0f)
+			{
+				AllDownloadsRetryingSince.store(FStatsCollector::GetCycles());
+				return false;
+			}
+
+			// When all chunks are retrying for a while the connection status is set to disconnected
+			const double SecondsSinceAllRetrying = FStatsCollector::CyclesToSeconds(FStatsCollector::GetCycles() - AllDownloadsRetryingSince.load(std::memory_order_relaxed));
+			return SecondsSinceAllRetrying >= Configuration.DisconnectedDelay;
+		}
 		
 		void OnDownloadProgress(int32 RequestId, uint64 BytesSoFar)
 		{
@@ -975,6 +1243,56 @@ namespace BuildPatchServices
 			WakeupMainThreadFn();
 		}
 
+		void UpdateCloudDirFailure(FCloudRead* Read)
+		{
+			// Mark this CDN as failed.
+			// Cap to prevent wrap. I think this is technically impossible due to the time it would take to get here but...
+			CloudDirFailureCount[Read->CloudDirUsed] = FMath::Min(CloudDirFailureCount[Read->CloudDirUsed] + 1, 100 << 20);
+			// Find who has failed the least, be sure to take equivalents in the initial specified order.
+			// We expect this to be like 5 entries.
+			const int32* MinFailedCloudDir = Algo::MinElement(CloudDirFailureCount);
+
+			const int32 MinAtIndex = MinFailedCloudDir - CloudDirFailureCount.GetData();
+
+			if (CurrentBestCloudDir != MinAtIndex)
+			{
+				CurrentBestCloudDir = MinAtIndex;
+				CurrentBestCloudDir = MinAtIndex;
+				const TCHAR* Reason = Read->bDecompressionSuccess ? TEXT("download") : TEXT("decompress");
+				UE_LOG(LogCloudChunkSource, Warning, TEXT("CDN %s failed %s, updating CDN selection to: %s"), 
+					*Configuration.CloudRoots[Read->CloudDirUsed],
+					Reason,
+					*Configuration.CloudRoots[MinAtIndex]);
+			}
+
+			MessagePump->SendMessage({ FGenericMessage::EType::CDNDownloadFailed, Read->DataId, Configuration.CloudRoots[Read->CloudDirUsed], Configuration.CloudRoots[MinAtIndex] });
+		}
+
+		void OnDownloadOrDecompressionFailed(FCloudRead* Read)
+		{
+			UpdateCloudDirFailure(Read);
+
+			ChunkSuccessTracker.AddFail();
+
+			DownloadRateLimiter.Refund(Read->ExpectedSize);
+
+			// Update retry.
+			Read->RetryNum++;
+			if (Configuration.MaxRetryCount >= 0 && Read->RetryNum >= Configuration.MaxRetryCount)
+			{
+				// Fail the request.
+				Read->CompleteFn.Execute(Read->DataId, false, true, Read->UserPtr);
+				delete Read;
+			}
+			else
+			{
+				// Set retry time and put the task back in the queue.
+				Read->RetryTime = FStatsCollector::GetSeconds() + GetRetryDelay(Configuration.RetryDelayTimes, Read->RetryNum);
+				FScopeLock _(&ReadQueueCS);
+				ReadQueue.Add(Read);
+			}
+		}
+
 		FConstructorCloudChunkSourceConfig Configuration;
 		TSharedRef<FDownloadDelegates, ESPMode::ThreadSafe> DownloadDelegates;
 		FDownloadProgressDelegate OnDownloadProgressDelegate;
@@ -994,11 +1312,13 @@ namespace BuildPatchServices
 		// Tracking health and connection state
 		std::atomic<int64> CyclesAtLastData = 0;
 		std::atomic_int32_t NextReadId = 0;
+		std::atomic<int64> AllDownloadsRetryingSince = -1.0f;
 
 		// Communication from download thread to processing thread.
 		FCriticalSection CompletedRequestsCS;
 		TMap<int32, FDownloadRef> CompletedDownloads; // Keyed off the download request id
 		TMap<int32, FChunkUriResponse> CompletedRequests; // Keyed off the read id
+		TMap<int32, FCloudRead*> CompletedDecompressions; // Keyed off the completed decompression task id
 
 		// If we start getting failures on our downloads, we track which ones
 		// fail and avoid them until everything goes bad. Initially we just hit the
@@ -1017,6 +1337,8 @@ namespace BuildPatchServices
 		// seconds of success...
 		FChunkSuccessRate ChunkSuccessTracker;
 
+		FDownloadRateLimiter DownloadRateLimiter;
+
 		// Change recognition
 		EBuildPatchDownloadHealth LastSeenHealth = EBuildPatchDownloadHealth::NUM_Values;
 		float LastSeenSuccessRate = -1.0f;
@@ -1089,7 +1411,6 @@ namespace BuildPatchServices
 		}
 		InFlightDownloads.Empty();
 	}
-
 	//
 	// Since we need to have a lot of requests in flight at the same time _and_ we need to
 	// synchronize across them to manage how many are in flight at any given moment, we have
@@ -1100,6 +1421,8 @@ namespace BuildPatchServices
 	{
 		int32 StartingDownloadCount = InFlightDownloads.Num();
 
+		DownloadRateLimiter.Tick();
+		
 		{
 			int32 PendingCount;
 			{
@@ -1117,11 +1440,34 @@ namespace BuildPatchServices
 		// Process completed downloads.
 		TMap<int32, FDownloadRef> FrameCompletedDownloads;
 		TMap<int32, FChunkUriResponse> FrameCompletedRequests;
+		TMap<int32, FCloudRead*> FrameCompletedDecompressions;
 		{
-			CompletedRequestsCS.Lock();
+			FScopeLock ScopeLock(&CompletedRequestsCS);
 			FrameCompletedRequests = MoveTemp(CompletedRequests);
 			FrameCompletedDownloads = MoveTemp(CompletedDownloads);
-			CompletedRequestsCS.Unlock();
+			FrameCompletedDecompressions = MoveTemp(CompletedDecompressions);
+		}
+
+		{
+			// Finalize completed FCloudRead;
+			for (const TPair<int32, FCloudRead*>& CompletedDecompression : FrameCompletedDecompressions)
+			{
+				FCloudRead* Read = CompletedDecompression.Value;
+				if (Read->bDecompressionSuccess)
+				{
+					ChunkSuccessTracker.AddSuccess();
+					CloudChunkSourceStat->OnDownloadSuccess(Read->DataId);
+					Read->CompleteFn.Execute(Read->DataId, false, false, Read->UserPtr);
+					delete Read;
+				}
+				else
+				{
+					// \todo this wants to know what the actual internal error was for some reason. idk why.
+					CloudChunkSourceStat->OnDownloadCorrupt(Read->DataId, Read->UrlUsed, EChunkLoadResult::HashCheckFailed);
+					OnDownloadOrDecompressionFailed(Read);
+				}
+
+			}
 		}
 
 		{
@@ -1149,7 +1495,7 @@ namespace BuildPatchServices
 					{
 						Read->UrlUsed = Response.Uri;
 					}
-					Read->RequestId = DownloadService->RequestFileWithHeaders(Read->UrlUsed, Response.AdditionalHeaders, OnDownloadCompleteDelegate, OnDownloadProgressDelegate);
+					DownloadService->RequestFileWithHeaders(Read->UrlUsed, Response.AdditionalHeaders, OnDownloadCompleteDelegate, OnDownloadProgressDelegate, &Read->RequestId);
 					InFlightDownloads.Add(Read->RequestId, Read);
 
 					CloudChunkSourceStat->OnDownloadRequested(Read->DataId);
@@ -1191,6 +1537,7 @@ namespace BuildPatchServices
 					//UE_LOG(LogTemp, Warning, TEXT("Rate: %.2f MB/s"), MegaBytesPerSecond);
 
 					CloudChunkSourceStat->OnReceivedDataUpdated(TotalBytes);
+					CloudChunkSourceStat->OnCloudDirectoryDataReceived(Configuration.CloudRoots[Read->CloudDirUsed], DownloadedData.Num());
 
 					FMemoryReader ReaderThunk(DownloadedData);
 
@@ -1202,10 +1549,8 @@ namespace BuildPatchServices
 
 					if (!Result)
 					{
-						// The header or chunk data was bad.
-						UE_LOG(LogCloudChunkSource, Error, TEXT("Failed to ValidateAndRead chunk %s, download size %d"), *WriteToString<64>(Read->DataId), DownloadedData.Num());
-						Read->CompleteFn.Execute(Read->DataId, false, true, Read->UserPtr);
-						delete Read;
+						// The header or chunk data was bad.				
+						bDownloadSuccess = false;
 					}
 					else
 					{
@@ -1221,109 +1566,37 @@ namespace BuildPatchServices
 					
 						// We either need to hash the chunk for validation or decompress it in to the destination buffer - don't
 						// block IO for this.
-						UE::Tasks::Launch(TEXT("CloudDecompressionAndHash"), [Read, Header, CloudChunkSourceStat = CloudChunkSourceStat, ChunkDataSerialization = ChunkDataSerialization, CompressedBuffer = MoveTemp(CompressedBuffer)]()
+						UE::Tasks::Launch(TEXT("CloudDecompressionAndHash"), [this, Read, Header, CompressedBuffer = MoveTemp(CompressedBuffer)]()
 							{
-								bool bDecompressSucceeded = ChunkDataSerialization->DecompressValidatedRead(Header, Read->Destination, CompressedBuffer);
-
-								if (bDecompressSucceeded)
+								Read->bDecompressionSuccess = ChunkDataSerialization->DecompressAndDecryptValidatedRead(Header, Read->Destination, CompressedBuffer);
 								{
-									CloudChunkSourceStat->OnDownloadSuccess(Read->DataId);
+									FScopeLock ScopeLock(&CompletedRequestsCS);
+									CompletedDecompressions.Add(Read->ReadId, Read);
 								}
-								else
-								{
-									// \todo this wants to know what the actual internal error was for some reason. idk why.
-									CloudChunkSourceStat->OnDownloadCorrupt(Read->DataId, Read->UrlUsed, EChunkLoadResult::HashCheckFailed);
-								}
-
-								Read->CompleteFn.Execute(Read->DataId, false, !bDecompressSucceeded, Read->UserPtr);
-								delete Read;
+								WakeupMainThreadFn();
 							}
 						);
 					}
-
-					ChunkSuccessTracker.AddSuccess();
 				}
 				else
 				{
 					CloudChunkSourceStat->OnDownloadFailed(Read->DataId, Read->UrlUsed);
+					CloudChunkSourceStat->OnCloudDirectoryError(Configuration.CloudRoots[Read->CloudDirUsed]);
 					UE_LOG(LogCloudChunkSource, Error, TEXT("FAILED: %s"), *Read->UrlUsed);
-					
-					// Mark this CDN as failed.
-					{
-						CloudDirFailureCount[Read->CloudDirUsed]++;
-						if (CloudDirFailureCount[Read->CloudDirUsed] > (100 << 20))
-						{
-							// Cap to prevent wrap. I think this is technically impossible due to the time it would take to 
-							// get here but...
-							CloudDirFailureCount[Read->CloudDirUsed] = 100 << 20;
-						}
-
-						// Find who has failed the least, be sure to take equivalents in the initial specified order.
-						// We expect this to be like 5 entries.
-						int32 MinFailCount = CloudDirFailureCount[0];
-						int32 MinAtIndex = 0;
-						for (int32 CloudDirSeek = 1; CloudDirSeek < CloudDirFailureCount.Num(); CloudDirSeek++)
-						{
-							if (CloudDirFailureCount[CloudDirSeek] < MinFailCount)
-							{
-								MinFailCount = CloudDirFailureCount[CloudDirSeek];
-								MinAtIndex = CloudDirSeek;
-							}
-						}
-
-						CurrentBestCloudDir = MinAtIndex;
-						if (Read->CloudDirUsed != MinAtIndex)
-						{
-							UE_LOG(LogCloudChunkSource, Warning, TEXT("CDN %s failed download, updating CDN selection to: %s"), *Configuration.CloudRoots[Read->CloudDirUsed], *Configuration.CloudRoots[MinAtIndex]);
-						}
-
-						MessagePump->SendMessage({FGenericMessage::EType::CDNDownloadFailed, Read->DataId, Configuration.CloudRoots[Read->CloudDirUsed], Configuration.CloudRoots[MinAtIndex]});
-					}
-					ChunkSuccessTracker.AddFail();
+				}
 
-					// Update retry.
-					Read->RetryNum++;
-					if (Configuration.MaxRetryCount >= 0 && Read->RetryNum >= Configuration.MaxRetryCount)
-					{
-						// Fail the request.
-						Read->CompleteFn.Execute(Read->DataId, false, true, Read->UserPtr);
-						delete Read;
-					}
-					else
-					{
-						// Set retry time and put the task back in the queue.
-						Read->RetryTime = FStatsCollector::GetSeconds() + GetRetryDelay(Configuration.RetryDelayTimes, Read->RetryNum);
-						FScopeLock _(&ReadQueueCS);
-						ReadQueue.Add(Read);
-					}
+				// Handle failed (note this also launches a retry on a bad serialization, not just download.
+				if(!bDownloadSuccess)
+				{
+					OnDownloadOrDecompressionFailed(Read);
 				}
 			} // end each completed
 		} // end process completed
 
 		// Update connection status and health.
 		EBuildPatchDownloadHealth DownloadHealth;
-		bool bAnyDownloadsRetrying = false;
 		{
-			bool bAllDownloadsRetrying = true;
-			{
-				FScopeLock _(&ReadQueueCS);
-				bAllDownloadsRetrying = ReadQueue.Num() > 0;
-				for (FCloudRead* Read : ReadQueue)
-				{
-					if (!Read->RetryNum)
-					{
-						bAllDownloadsRetrying = false;
-					}
-					else
-					{
-						bAnyDownloadsRetrying = true;
-					}
-				}
-			}
-
-			const double SecondsSinceData = FStatsCollector::CyclesToSeconds(FStatsCollector::GetCycles() - CyclesAtLastData.load(std::memory_order_relaxed));
-			const bool bReportAsDisconnected = (bAllDownloadsRetrying && SecondsSinceData > Configuration.DisconnectedDelay);
-
+			const bool bReportAsDisconnected = GetDisconnectionStatus();
 			const float SuccessRate = ChunkSuccessTracker.GetOverall();
 			EBuildPatchDownloadHealth OverallDownloadHealth = GetDownloadHealth(bReportAsDisconnected, Configuration.HealthPercentages, SuccessRate);
 			if (LastSeenHealth != OverallDownloadHealth)
@@ -1360,6 +1633,7 @@ namespace BuildPatchServices
 			OldestReadId = NewestReadId;
 		}
 
+		uint64 WaitingForTokensUntilMs = 0;
 		while (bStartNewDownloads)
 		{
 			const uint32 MaxDownloads = FMath::Min(InMaxDownloads, FMath::Max(1U, DownloadCount->GetAdjustedCount(InFlightDownloads.Num(), DownloadHealth)));
@@ -1369,6 +1643,8 @@ namespace BuildPatchServices
 			}
 
 			FCloudRead* Read = nullptr;
+			uint64 DownloadSize;
+
 			{
 				FScopeLock _(&ReadQueueCS);
 
@@ -1378,6 +1654,13 @@ namespace BuildPatchServices
 				}
 
 				Read = ReadQueue[0];
+				DownloadSize = ManifestSet->GetDownloadSize(Read->DataId);
+				uint64 MsUntilScheduled = DownloadRateLimiter.TrySchedule(DownloadSize);
+				if (MsUntilScheduled)
+				{
+					WaitingForTokensUntilMs = MsUntilScheduled;
+					break;
+				}
 				ReadQueue.RemoveAt(0);
 			}
 
@@ -1385,14 +1668,17 @@ namespace BuildPatchServices
 			// due to appending auth keys to the URL itself, or adding auth headers. This can be (very) latent or instantaneous.
 			FChunkUriRequest ChunkUriRequest;
 			ChunkUriRequest.CloudDirectory = Configuration.CloudRoots[CurrentBestCloudDir];
-			ChunkUriRequest.RelativePath = ManifestSet->GetDataFilename(Read->DataId);
-			ChunkUriRequest.RelativePath.RemoveFromStart(TEXT("/"));
-
+			ChunkUriRequest.ChunkCloudSubdirectory = ManifestSet->GetChunkCloudSubdirectory(Read->DataId);
+			ChunkUriRequest.ChunkCloudSubdirectory.RemoveFromStart(TEXT("/"));
+			ChunkUriRequest.ChunkPath = ManifestSet->GetChunkFilename(Read->DataId);
+			ChunkUriRequest.UniqueBuildId = ManifestSet->GetChunkUniqueBuildId(Read->DataId);
+			ChunkUriRequest.bChunkFromDelta = ManifestSet->GetChunkFromDeltaStatus(Read->DataId);
+			
 			Read->CloudDirUsed = CurrentBestCloudDir;
 			Read->bQueuedForDownload = true;
 			Read->ExpectedSize = ManifestSet->GetDownloadSize(Read->DataId);
 			Read->SecondsAtRequested = FStatsCollector::GetSeconds();
-			Read->UrlUsed = ChunkUriRequest.CloudDirectory / ChunkUriRequest.RelativePath;
+			Read->UrlUsed = ChunkUriRequest.CloudDirectory / ChunkUriRequest.ChunkCloudSubdirectory / ChunkUriRequest.ChunkPath;
 			UE_LOG(LogCloudChunkSource, VeryVerbose, TEXT("Downloading: %s"), *Read->UrlUsed);
 
 			InFlightURLRequests.Add(Read->ReadId, Read);
@@ -1410,7 +1696,7 @@ namespace BuildPatchServices
 				// If no message handler was registered to do any modification of the URL, start the download immediately.
 				InFlightURLRequests.Remove(Read->ReadId);
 
-				Read->RequestId = DownloadService->RequestFileWithHeaders(Read->UrlUsed, {}, OnDownloadCompleteDelegate, OnDownloadProgressDelegate);
+				DownloadService->RequestFileWithHeaders(Read->UrlUsed, {}, OnDownloadCompleteDelegate, OnDownloadProgressDelegate, &Read->RequestId);
 				InFlightDownloads.Add(Read->RequestId, Read);
 
 				CloudChunkSourceStat->OnDownloadRequested(Read->DataId);
@@ -1482,6 +1768,12 @@ namespace BuildPatchServices
 			}
 		}
 
+		// If we aren't kicking reads because we don't have the quota, we need to check back soon.
+		if (WaitingForTokensUntilMs)
+		{
+			WaitTimeMs = WaitingForTokensUntilMs;
+		}
+
 		OutTimeToNextTickMs = WaitTimeMs;
 	}
 
diff --git a/Engine/Source/Runtime/Online/BuildPatchServices/Private/Installer/CloudChunkSource.h b/Engine/Source/Runtime/Online/BuildPatchServices/Private/Installer/CloudChunkSource.h
--- a/Engine/Source/Runtime/Online/BuildPatchServices/Private/Installer/CloudChunkSource.h
+++ b/Engine/Source/Runtime/Online/BuildPatchServices/Private/Installer/CloudChunkSource.h
@@ -39,6 +39,10 @@ namespace BuildPatchServices
 		// When all requests are failing, how many seconds before a success until we determine the state as disconnected.
 		float DisconnectedDelay = 5.0f;
 
+		// Maximum download speed in bytes per second.
+		// Infinite can be specified as 0.
+		uint64 DownloadRateLimitBps = 0;
+
 		FConstructorCloudChunkSourceConfig(const TArray<FString>& InCloudRoots)
 		: CloudRoots(InCloudRoots)
 		{
@@ -68,6 +72,8 @@ namespace BuildPatchServices
 		// Pass in total download requirements for the install (including already completed downloads)
 		virtual void PostRequiredByteCount(uint64 InDownloadExpected) = 0;
 
+		virtual void SetDownloadBytesPerSecondLimit(uint64 BytesPerSecondLimit) = 0;
+
 		static IConstructorCloudChunkSource* CreateCloudSource(FConstructorCloudChunkSourceConfig&& Configuration, IDownloadService* InDownloadService, 
 			IChunkDataSerialization* InChunkDataSerialization, IDownloadConnectionCount* InDownloadCount, IMessagePump* InMessagePump, ICloudChunkSourceStat* InCloudChunkSourceStat, 
 			IBuildManifestSet* InManifestSet);
@@ -83,6 +89,12 @@ namespace BuildPatchServices
 	public:
 		virtual ~ICloudChunkSource() = default;
 
+		/**
+		 * Update the download rate limit.
+		 * @param InDownloadRateLimitBps    The new maximum download rate in bytes per second. Infinite can be specified as 0.
+		 */
+		virtual void SetDownloadBytesPerSecondLimit(uint64 InDownloadRateLimitBps) = 0;
+
 		/**
 		 * Should be called from the desired thread if bRunOwnThread is false.
 		 * It will block until until aborted by the owning system.
@@ -100,12 +112,17 @@ namespace BuildPatchServices
 		// The maximum number of times that a single chunk should retry, before registering a fatal error.
 		// Infinite can be specified as < 0.
 		int32 MaxRetryCount = 6;
+		// Number of retries to enable logging of http warnings.
+		int32 RetriesToEnableVerbose;
 		// The minimum number of chunks to request ahead of what is required, depending on store slack.
 		int32 PreFetchMinimum  = 16;
 		// The maximum number of chunks to request ahead of what is required, depending on store slack.
 		// This shouldn't be too much because we'll queue too many on the same CDN which might not be good and we have to chew 
 		// though a ton of requests before we swap over, and we only really want there to be one queued + one ready to go per connection anyway.
 		int32 PreFetchMaximum = 32;
+		// Maximum download speed in bytes per second.
+		// Infinite can be specified as 0.
+		uint64 DownloadRateLimitBps = 0;
 		// Array of times in seconds, representing the time between each retry upon failure. The last entry will be used
 		// indefinitely once it is reached.
 		TArray<float> RetryDelayTimes;
@@ -130,6 +147,7 @@ namespace BuildPatchServices
 		 */
 		FCloudSourceConfig(const TArray<FString>& InCloudRoots)
 			: CloudRoots(InCloudRoots)
+			, RetriesToEnableVerbose(MaxRetryCount - 1)
 		{
 			const float RetryFloats[] = {0.5f, 1.0f, 1.0f, 3.0f, 3.0f, 10.0f, 10.0f, 20.0f, 20.0f, 30.0f};
 			RetryDelayTimes.Empty(UE_ARRAY_COUNT(RetryFloats));
@@ -259,10 +277,35 @@ namespace BuildPatchServices
 		 */
 		virtual void OnActiveRequestCountUpdated(uint32 RequestCount) = 0;
 
+		/**
+		 * Called once per second whenever there has been an active request.
+		 * @param AverageRequestCount      The average number of active requests that there has currently been.
+		 */
+		virtual void OnActiveRequestCountAvgUpdated(int32 AverageRequestCount) = 0;
+
 		/**
 		 * Called when a batch of chunks are added and accepted via IChunkSource::AddRuntimeRequirements.
 		 * @param ChunkIds  The ids of each chunk.
 		 */
 		virtual void OnAcceptedNewRequirements(const TSet<FGuid>& ChunkIds) = 0;
+
+		/**
+		 * Called when a chunk is decrypted.
+		 * @param DecryptionCycles  The amount of cycles spent on decryption a chunk.
+		 */
+		virtual void OnChunkDecrypted(uint64 DecryptionCycles) = 0;
+
+		/**
+		 * Called when new data has been received from the given cloud directory.
+		 * @param CloudDirectory  The identifier for this cloud directory, e.g. base URL.
+		 * @param ReceivedData    The number of bytes received from this cloud directory.
+		 */
+		virtual void OnCloudDirectoryDataReceived(const FString& CloudDirectory, uint64 ReceivedData) = 0;
+
+		/**
+		 * Called when a download from the given cloud directory has failed.
+		 * @param CloudDirectory  The identifier for this cloud directory, e.g. base URL.
+		 */
+		virtual void OnCloudDirectoryError(const FString& CloudDirectory) = 0;
 	};
 }
diff --git a/Engine/Source/Runtime/Online/BuildPatchServices/Private/Installer/Controllable.h b/Engine/Source/Runtime/Online/BuildPatchServices/Private/Installer/Controllable.h
--- a/Engine/Source/Runtime/Online/BuildPatchServices/Private/Installer/Controllable.h
+++ b/Engine/Source/Runtime/Online/BuildPatchServices/Private/Installer/Controllable.h
@@ -21,5 +21,10 @@ namespace BuildPatchServices
 		 * Called to instruct the class to cease all activity, and perform any shutdown.
 		 */
 		virtual void Abort() = 0;
+
+		/**
+		 * Return to default state without destructing the object.
+		 */
+		virtual void Reset() = 0;
 	};
 }
diff --git a/Engine/Source/Runtime/Online/BuildPatchServices/Private/Installer/DownloadConnectionCount.cpp b/Engine/Source/Runtime/Online/BuildPatchServices/Private/Installer/DownloadConnectionCount.cpp
--- a/Engine/Source/Runtime/Online/BuildPatchServices/Private/Installer/DownloadConnectionCount.cpp
+++ b/Engine/Source/Runtime/Online/BuildPatchServices/Private/Installer/DownloadConnectionCount.cpp
@@ -33,10 +33,9 @@ namespace BuildPatchServices
 
 	public:
 		FDownloadConnectionCount(FDownloadConnectionCountConfig InConfiguration, IDownloadServiceStatistics* InDownloadStatistics);
-		virtual ~FDownloadConnectionCount() override;
-
-		// IDownloadConnectionCount Interface
-		virtual uint32 GetAdjustedCount(uint32 NumProcessing, EBuildPatchDownloadHealth CurrentHealth) override;
+		virtual ~FDownloadConnectionCount();
+		virtual uint32 GetAdjustedCount(uint32 NumProcessing, EBuildPatchDownloadHealth CurrentHealth);
+		void SetConnectionScaling(bool bEnableScaling) override;
 
 	private:
 		EDownloadSpeedStatus GetSpeedStatus(double Speed) const;
@@ -53,6 +52,7 @@ namespace BuildPatchServices
 		void UpdateSpeedStatus(EDownloadSpeedStatus SpeedStatus);
 		void UpdateAllowance(int32 Adjustment);
 		void UpdateHealth(EBuildPatchDownloadHealth CurrentHealth);
+		void Reset();
 
 	private:
 		IDownloadServiceStatistics* DownloadStatistics;
@@ -66,6 +66,8 @@ namespace BuildPatchServices
 		uint32 AveragesSeen;
 		const uint32 MinimumAveragesBeforeHighSet = 4;
 		const uint32 MaximumAllowanceExcess = 16;
+		TAtomic<bool> bConnectionScalingEnabled;
+		bool bPreviousConnectionScalingEnabledStatus;
 	};
 
 	FDownloadConnectionCount::FDownloadConnectionCount(FDownloadConnectionCountConfig InConfiguration, IDownloadServiceStatistics* InDownloadStatistics)
@@ -75,11 +77,26 @@ namespace BuildPatchServices
 		, RequestAllowance(Configuration.FallbackCount)
 		, ConsecutiveIncreases(0U)
 		, ConsecutiveFailedHealth(0U)
-		, HighBandwidthLevel(0.0L)
-		, LowBandwidthLevel(0.0L)
+		, HighBandwidthLevel(0.0)
+		, LowBandwidthLevel(0.0)
 		, AveragesSeen(0U)
+		, bConnectionScalingEnabled(!Configuration.bDisableConnectionScaling)
+		, bPreviousConnectionScalingEnabledStatus(!Configuration.bDisableConnectionScaling)
 	{
+	}
 
+	void FDownloadConnectionCount::Reset()
+	{
+		ConsecutiveDecreases = 0U;
+		RequestAllowance = Configuration.FallbackCount;
+		ConsecutiveIncreases = 0U;
+		ConsecutiveFailedHealth = 0U;
+		HighBandwidthLevel = 0.0;
+		LowBandwidthLevel = 0.0;
+		AveragesSeen = 0U;
+
+		// Calling with 0 resets the internal values.
+		DownloadStatistics->GetImmediateAverageSpeedPerRequest(0);
 	}
 
 	FDownloadConnectionCount::~FDownloadConnectionCount()
@@ -88,10 +105,16 @@ namespace BuildPatchServices
 
 	uint32 FDownloadConnectionCount::GetAdjustedCount(uint32 NumProcessing, EBuildPatchDownloadHealth CurrentHealth)
 	{
-		if (Configuration.bDisableConnectionScaling)
+		if (bConnectionScalingEnabled == false)
 		{
+			bPreviousConnectionScalingEnabledStatus = false;
 			return Configuration.FallbackCount;
 		}
+		else if (!bPreviousConnectionScalingEnabledStatus)
+		{
+			Reset();
+		}
+		bPreviousConnectionScalingEnabledStatus = true;
 
 		uint32 OldCount = GetAllowance();
 		// Prevent request allowance getting too far ahead of what is being processed currently.
@@ -158,6 +181,11 @@ namespace BuildPatchServices
 		return GetAllowance();
 	}
 
+	void FDownloadConnectionCount::SetConnectionScaling(bool bEnableScaling)
+	{
+		bConnectionScalingEnabled = bEnableScaling;
+	}
+
 	bool FDownloadConnectionCount::IsHealthAcceptable(EBuildPatchDownloadHealth CurrentHealth) const
 	{
 		bool bIsOk = true;
diff --git a/Engine/Source/Runtime/Online/BuildPatchServices/Private/Installer/DownloadConnectionCount.h b/Engine/Source/Runtime/Online/BuildPatchServices/Private/Installer/DownloadConnectionCount.h
--- a/Engine/Source/Runtime/Online/BuildPatchServices/Private/Installer/DownloadConnectionCount.h
+++ b/Engine/Source/Runtime/Online/BuildPatchServices/Private/Installer/DownloadConnectionCount.h
@@ -16,13 +16,20 @@ namespace BuildPatchServices
 	class IDownloadConnectionCount
 	{
 	public:
-		/*
+		/**
 		 * Get the number of simultaneous download connections
 		 * @param NumProcessing -- the current number of downloads in-flight
 		 * @param CurrentHealth -- percentage of downloads that have failed since last call to this function, in terms of "health"
 		 * @return number of connections that should be approximately optimal.
-		*/
+		 */
 		virtual uint32 GetAdjustedCount(uint32 NumProcessing, EBuildPatchDownloadHealth CurrentHealth) = 0;
+
+		/**
+		 * Set whether connection scaling is enabled.
+		 * @param bEnableScaling -- Enables download scaling, true to enable.
+		 */
+		virtual void SetConnectionScaling(bool bEnableScaling) = 0;
+
 		virtual ~IDownloadConnectionCount() = default;
 	};
 
@@ -89,4 +96,4 @@ namespace BuildPatchServices
 		*/
 		static IDownloadConnectionCount* Create(FDownloadConnectionCountConfig Configuration, IDownloadServiceStatistics* InDownloadStatistics);
 	};
-}
+}
\ No newline at end of file
diff --git a/Engine/Source/Runtime/Online/BuildPatchServices/Private/Installer/DownloadService.cpp b/Engine/Source/Runtime/Online/BuildPatchServices/Private/Installer/DownloadService.cpp
--- a/Engine/Source/Runtime/Online/BuildPatchServices/Private/Installer/DownloadService.cpp
+++ b/Engine/Source/Runtime/Online/BuildPatchServices/Private/Installer/DownloadService.cpp
@@ -100,13 +100,13 @@ namespace BuildPatchServices
 		: public TSharedFromThis<FDownloadServiceImpl>
 	{
 	public:
-		FDownloadServiceImpl(IHttpManager* HttpManager, IFileSystem* FileSystem, IDownloadServiceStat* DownloadServiceStat, IInstallerAnalytics* InstallerAnalytics);
+		FDownloadServiceImpl(TSharedPtr<IHttpManager> HttpManager, TSharedPtr<IFileSystem> FileSystem, TSharedPtr<IDownloadServiceStat> DownloadServiceStat, TSharedPtr<IInstallerAnalytics> InstallerAnalytics);
 		~FDownloadServiceImpl();
 
 		void CancelAllRequests();
 
 		// IDownloadService interface begin.
-		int32 RequestFileWithHeaders(const FString& FileUri, const TMap<FString, FString>& RequestHeaders, const FDownloadCompleteDelegate& OnCompleteDelegate, const FDownloadProgressDelegate& OnProgressDelegate);
+		bool RequestFileWithHeaders(const FString& FileUri, const TMap<FString, FString>& RequestHeaders, const FDownloadCompleteDelegate& OnCompleteDelegate, const FDownloadProgressDelegate& OnProgressDelegate, int32* OutRequestId = nullptr);
 		void RequestCancel(int32 RequestId);
 		void RequestAbandon(int32 RequestId);
 		// IDownloadService interface end.
@@ -117,10 +117,10 @@ namespace BuildPatchServices
 		IDownloadServiceStat::FDownloadRecord MakeDownloadRecord(int32 RequestId, FString Uri);
 
 	private:
-		IHttpManager* HttpManager = nullptr;
-		IFileSystem* FileSystem = nullptr;
-		IDownloadServiceStat* DownloadServiceStat = nullptr;
-		IInstallerAnalytics* InstallerAnalytics = nullptr;
+		TSharedPtr<IHttpManager> HttpManager;
+		TSharedPtr<IFileSystem> FileSystem;
+		TSharedPtr<IDownloadServiceStat> DownloadServiceStat;
+		TSharedPtr<IInstallerAnalytics> InstallerAnalytics;
 
 		std::atomic<int32> RequestIdCounter;
 
@@ -136,7 +136,7 @@ namespace BuildPatchServices
 		: public IDownloadService
 	{
 	public:
-		FDownloadService(IHttpManager* HttpManager, IFileSystem* FileSystem, IDownloadServiceStat* DownloadServiceStat, IInstallerAnalytics* InstallerAnalytics)
+		FDownloadService(TSharedPtr<IHttpManager> HttpManager, TSharedPtr<IFileSystem> FileSystem, TSharedPtr<IDownloadServiceStat> DownloadServiceStat, TSharedPtr<IInstallerAnalytics> InstallerAnalytics)
 			: Impl(MakeShared<FDownloadServiceImpl>(HttpManager, FileSystem, DownloadServiceStat, InstallerAnalytics))
 		{}
 		~FDownloadService()
@@ -145,14 +145,14 @@ namespace BuildPatchServices
 		}
 
 		// IDownloadService interface begin.
-		virtual int32 RequestFile(const FString& FileUri, const FDownloadCompleteDelegate& OnCompleteDelegate, const FDownloadProgressDelegate& OnProgressDelegate) override
+		virtual bool RequestFile(const FString& FileUri, const FDownloadCompleteDelegate& OnCompleteDelegate, const FDownloadProgressDelegate& OnProgressDelegate, int32* OutRequestId = nullptr) override
 		{
 			TMap<FString, FString> RequestHeaders;
-			return Impl->RequestFileWithHeaders(FileUri, RequestHeaders, OnCompleteDelegate, OnProgressDelegate);
+			return Impl->RequestFileWithHeaders(FileUri, RequestHeaders, OnCompleteDelegate, OnProgressDelegate, OutRequestId);
 		}
-		int32 RequestFileWithHeaders(const FString& FileUri, const TMap<FString, FString>& RequestHeaders, const FDownloadCompleteDelegate& OnCompleteDelegate, const FDownloadProgressDelegate& OnProgressDelegate)
+		virtual bool RequestFileWithHeaders(const FString& FileUri, const TMap<FString, FString>& RequestHeaders, const FDownloadCompleteDelegate& OnCompleteDelegate, const FDownloadProgressDelegate& OnProgressDelegate, int32* OutRequestId = nullptr) override
 		{
-			return Impl->RequestFileWithHeaders(FileUri, RequestHeaders, OnCompleteDelegate, OnProgressDelegate);
+			return Impl->RequestFileWithHeaders(FileUri, RequestHeaders, OnCompleteDelegate, OnProgressDelegate, OutRequestId);
 		}
 		virtual void RequestCancel(int32 RequestId) override { Impl->RequestCancel(RequestId); }
 		virtual void RequestAbandon(int32 RequestId) override { Impl->RequestAbandon(RequestId); }
@@ -162,7 +162,7 @@ namespace BuildPatchServices
 		TSharedRef<FDownloadServiceImpl> Impl;
 	};
 
-	FDownloadServiceImpl::FDownloadServiceImpl(IHttpManager* InHttpManager, IFileSystem* InFileSystem, IDownloadServiceStat* InDownloadServiceStat, IInstallerAnalytics* InInstallerAnalytics)
+	FDownloadServiceImpl::FDownloadServiceImpl(TSharedPtr<IHttpManager> InHttpManager, TSharedPtr<IFileSystem> InFileSystem, TSharedPtr<IDownloadServiceStat> InDownloadServiceStat, TSharedPtr<IInstallerAnalytics> InInstallerAnalytics)
 		: HttpManager(InHttpManager)
 		, FileSystem(InFileSystem)
 		, DownloadServiceStat(InDownloadServiceStat)
@@ -191,9 +191,14 @@ namespace BuildPatchServices
 		// callbacks and removal handled by completion delegates
 	}
 
-	int32 FDownloadServiceImpl::RequestFileWithHeaders(const FString& FileUri, const TMap<FString, FString>& RequestHeaders, const FDownloadCompleteDelegate& OnCompleteDelegate, const FDownloadProgressDelegate& OnProgressDelegate)
+	bool FDownloadServiceImpl::RequestFileWithHeaders(const FString& FileUri, const TMap<FString, FString>& RequestHeaders, const FDownloadCompleteDelegate& OnCompleteDelegate, const FDownloadProgressDelegate& OnProgressDelegate, int32* OutRequestId)
 	{
 		const int32 RequestId = ++RequestIdCounter;
+		if (OutRequestId)
+		{
+			*OutRequestId = RequestId;
+		}
+
 
 		// Save the delegates.
 		{
@@ -244,9 +249,10 @@ namespace BuildPatchServices
 				InstallerAnalytics->TrackRequest(Request);
 
 				DownloadRecord.bSuccess = bSucceeded;
-				DownloadRecord.ResponseCode = Response.IsValid() ? Response->GetResponseCode() : INDEX_NONE;
+				DownloadRecord.ResponseCode = Response.IsValid() ? Response->GetResponseCode() : (int32)Request->GetFailureReason();
 				DownloadRecord.SpeedRecord.CyclesEnd = FStatsCollector::GetCycles();
 				DownloadRecord.SpeedRecord.Size = Response.IsValid() ? Response->GetContent().Num() : 0;
+				DownloadRecord.FailureReason = Request->GetFailureReason();
 				
 				FDownloadDelegates Delegates;
 				bool bFoundDelegates = false;
@@ -275,16 +281,23 @@ namespace BuildPatchServices
 		{
 			// Load file from drive/network.
 			TUniquePtr<FFileRequest> FileRequest = MakeUnique<FFileRequest>();
-			FileRequest->Task = MakeFileLoadTask(RequestId, FileUri, FileRequest.Get());
-
+			FFileRequest* FileRequestPtr = FileRequest.Get();
 			{
 				FScopeLock ScopeLock(&ActiveRequestsCS);
 				ActiveFileRequests.Add(RequestId, MoveTemp(FileRequest));
 			}
+
+			UE::Tasks::FTask FileLoadTask = MakeFileLoadTask(RequestId, FileUri, FileRequestPtr);
+			{
+				FScopeLock ScopeLock(&ActiveRequestsCS);
+				if (ActiveFileRequests.Contains(RequestId))
+				{
+					FileRequestPtr->Task = MoveTemp(FileLoadTask);
+				}
+			}
 		}
 		DownloadServiceStat->OnDownloadStarted(RequestId, FileUri);
-
-		return RequestId;
+		return true;
 	}
 
 	void FDownloadServiceImpl::RequestCancel(int32 RequestId)
@@ -337,10 +350,9 @@ namespace BuildPatchServices
 						FileDataArray.Reset();
 						FileDataArray.AddUninitialized(FileSize);
 						int64 BytesRead = 0;
-						bool bIsCanceled = false;
-						while (BytesRead < FileSize && !FileRequest->ShouldCancel.IsCanceled())
+						while (BytesRead < FileSize)
 						{
-							bIsCanceled = FileRequest->ShouldCancel.IsCanceled();
+							bool bIsCanceled = FileRequest->ShouldCancel.IsCanceled();
 							if (bIsCanceled)
 							{
 								bSuccess = false;
@@ -411,7 +423,7 @@ namespace BuildPatchServices
 		return DownloadRecord;
 	}
 
-	IDownloadService* FDownloadServiceFactory::Create(IHttpManager* HttpManager, IFileSystem* FileSystem, IDownloadServiceStat* DownloadServiceStat, IInstallerAnalytics* InstallerAnalytics)
+	IDownloadService* FDownloadServiceFactory::Create(TSharedPtr<IHttpManager> HttpManager, TSharedPtr<IFileSystem> FileSystem, TSharedPtr<IDownloadServiceStat> DownloadServiceStat, TSharedPtr<IInstallerAnalytics> InstallerAnalytics)
 	{
 		check(HttpManager != nullptr);
 		check(FileSystem != nullptr);
diff --git a/Engine/Source/Runtime/Online/BuildPatchServices/Private/Installer/DownloadService.h b/Engine/Source/Runtime/Online/BuildPatchServices/Private/Installer/DownloadService.h
--- a/Engine/Source/Runtime/Online/BuildPatchServices/Private/Installer/DownloadService.h
+++ b/Engine/Source/Runtime/Online/BuildPatchServices/Private/Installer/DownloadService.h
@@ -4,6 +4,7 @@
 #include "CoreMinimal.h"
 #include "Containers/Ticker.h"
 #include "Common/SpeedRecorder.h"
+#include "Interfaces/IHttpBase.h"
 
 namespace BuildPatchServices
 {
@@ -73,12 +74,13 @@ namespace BuildPatchServices
 		 * @param FileUri               The uri for the file request. For http(s), this should begin with the protocol.
 		 * @param OnCompleteDelegate    The delegate that will be called with the completed download.
 		 * @param OnProgressDelegate    The delegate that will be called with updates to the ongoing download.
-		 * @return the request id for this download, which will match that provided to the delegates.
+		 * @param OutRequestId          The request id for this download, which will match that provided to the delegates.
+		 * @return if the request was queued successfully.
 		 */
-		virtual int32 RequestFile(const FString& FileUri, const FDownloadCompleteDelegate& OnCompleteDelegate, const FDownloadProgressDelegate& OnProgressDelegate) = 0;
+		virtual bool RequestFile(const FString& FileUri, const FDownloadCompleteDelegate& OnCompleteDelegate, const FDownloadProgressDelegate& OnProgressDelegate, int32* OutRequestId = nullptr) = 0;
 
 		/** As per RequestFile, except adds the given headers to the request as a <Name, Value> pair. */
-		virtual int32 RequestFileWithHeaders(const FString& FileUri, const TMap<FString, FString>& RequestHeaders, const FDownloadCompleteDelegate& OnCompleteDelegate, const FDownloadProgressDelegate& OnProgressDelegate) = 0;
+		virtual bool RequestFileWithHeaders(const FString& FileUri, const TMap<FString, FString>& RequestHeaders, const FDownloadCompleteDelegate& OnCompleteDelegate, const FDownloadProgressDelegate& OnProgressDelegate, int32* OutRequestId = nullptr) = 0;
 
 		/**
 		 * Requests the cancellation of a requested file.
@@ -107,7 +109,7 @@ namespace BuildPatchServices
 		 * @param InstallerAnalytics    The analytics implementation for tracking HTTP requests.
 		 * @return the new IDownloadService instance created.
 		 */
-		static IDownloadService* Create(IHttpManager* HttpManager, IFileSystem* FileSystem, IDownloadServiceStat* DownloadServiceStat, IInstallerAnalytics* InstallerAnalytics);
+		static IDownloadService* Create(TSharedPtr<IHttpManager> HttpManager, TSharedPtr<IFileSystem> FileSystem, TSharedPtr<IDownloadServiceStat> DownloadServiceStat, TSharedPtr<IInstallerAnalytics> InstallerAnalytics);
 	};
 
 	/**
@@ -134,6 +136,8 @@ namespace BuildPatchServices
 			int32 ResponseCode;
 			// The timing record.
 			ISpeedRecorder::FRecord SpeedRecord;
+			// If failed, the reason why the request failed
+			EHttpFailureReason FailureReason;
 		};
 
 		/**
diff --git a/Engine/Source/Runtime/Online/BuildPatchServices/Private/Installer/FileAttribution.cpp b/Engine/Source/Runtime/Online/BuildPatchServices/Private/Installer/FileAttribution.cpp
--- a/Engine/Source/Runtime/Online/BuildPatchServices/Private/Installer/FileAttribution.cpp
+++ b/Engine/Source/Runtime/Online/BuildPatchServices/Private/Installer/FileAttribution.cpp
@@ -5,6 +5,8 @@
 #include "BuildPatchProgress.h"
 #include "IBuildManifestSet.h"
 #include "BuildPatchUtil.h"
+#include "Core/AsyncHelpers.h"
+#include <atomic>
 
 namespace BuildPatchServices
 {
@@ -19,6 +21,7 @@ namespace BuildPatchServices
 		// IControllable interface begin.
 		virtual void SetPaused(bool bInIsPaused) override;
 		virtual void Abort() override;
+		virtual void Reset() override;
 		// IControllable interface end.
 
 		// IFileAttribution interface begin.
@@ -26,9 +29,9 @@ namespace BuildPatchServices
 		// IFileAttribution interface end.
 
 	private:
-		FString SelectFullFilePath(const FString& BuildFile);
-		bool HasSameAttributes(const FFileManifest* NewFileManifest, const FFileManifest* OldFileManifest);
-		void SetupFileAttributes(const FString& FilePath, const FFileManifest& FileManifest, bool bForce);
+		FString SelectFullFilePath(const FString& BuildFile) const;
+		bool HasSameAttributes(const FFileManifest* NewFileManifest, const FFileManifest* OldFileManifest) const;
+		void SetupFileAttributes(const FString& FilePath, const FFileManifest& FileManifest, bool bForce) const;
 
 	private:
 		IFileSystem* FileSystem;
@@ -69,43 +72,56 @@ namespace BuildPatchServices
 		bShouldAbort = true;
 	}
 
+	void FFileAttribution::Reset()
+	{
+		bShouldAbort = false;
+	}
+
 	bool FFileAttribution::ApplyAttributes()
 	{
 		// We need to set attributes for all files in the new build that require it
 		TSet<FString> BuildFileList;
 		ManifestSet->GetExpectedFiles(BuildFileList);
 		BuildProgress->SetStateProgress(EBuildPatchState::SettingAttributes, 0.0f);
-		BuildFileList.Sort(TLess<FString>());
-		int32 BuildFileIdx = INDEX_NONE;
+
+		FTaskCoresScheduler<void> ApplyAttributesTasks = BuildPatchServices::CreateTaskSchedulerLimitedByCoresNumber(EAsyncExecution::TaskGraph);
+
+		std::atomic<int32> BuildFileIdx{ 0 };
+		float FilesNum = BuildFileList.Num();
 		for (const FString& BuildFile : BuildFileList)
 		{
-			++BuildFileIdx;
 			if (bShouldAbort)
 			{
 				break;
 			}
-			const FFileManifest* NewFileManifest = ManifestSet->GetNewFileManifest(BuildFile);
-			const FFileManifest* OldFileManifest = ManifestSet->GetCurrentFileManifest(BuildFile);
-			const bool bForce = ManifestSet->IsFileRepairAction(BuildFile);
-			bool bHasChanged = bForce || (TouchedFiles.Contains(BuildFile) && !HasSameAttributes(NewFileManifest, OldFileManifest));
-			if (NewFileManifest != nullptr && bHasChanged)
-			{
-				SetupFileAttributes(SelectFullFilePath(BuildFile), *NewFileManifest, bForce);
-			}
-			BuildProgress->SetStateProgress(EBuildPatchState::SettingAttributes, BuildFileIdx / float(BuildFileList.Num()));
+			auto TaskFunc = [&]()
+				{
+					const FFileManifest* NewFileManifest = ManifestSet->GetNewFileManifest(BuildFile);
+					const FFileManifest* OldFileManifest = ManifestSet->GetCurrentFileManifest(BuildFile);
+					const bool bForce = ManifestSet->IsFileRepairAction(BuildFile);
+					bool bHasChanged = bForce || (TouchedFiles.Contains(BuildFile) && !HasSameAttributes(NewFileManifest, OldFileManifest));
+					if (NewFileManifest != nullptr && bHasChanged)
+					{
+						SetupFileAttributes(SelectFullFilePath(BuildFile), *NewFileManifest, bForce);
+					}
+					BuildProgress->SetStateProgress(EBuildPatchState::SettingAttributes, float(BuildFileIdx.fetch_add(1)) / FilesNum);
+				};
+			ApplyAttributesTasks.ScheduleTask(MoveTemp(TaskFunc));
 			// Wait while paused
 			while (bIsPaused && !bShouldAbort)
 			{
 				FPlatformProcess::Sleep(0.5f);
 			}
 		}
+
+		ApplyAttributesTasks.WaitAll();
 		BuildProgress->SetStateProgress(EBuildPatchState::SettingAttributes, 1.0f);
 
 		// We don't fail on this step currently
 		return true;
 	}
 
-	FString FFileAttribution::SelectFullFilePath(const FString& BuildFile)
+	FString FFileAttribution::SelectFullFilePath(const FString& BuildFile) const
 	{
 		FString InstallFilename;
 		if (bUseStageDirectory)
@@ -121,7 +137,7 @@ namespace BuildPatchServices
 		return FBuildPatchUtils::ResolveInstallationFileName(InstallDirectory, PerFileSubdirectories, BuildFile);
 	}
 
-	bool FFileAttribution::HasSameAttributes(const FFileManifest* NewFileManifest, const FFileManifest* OldFileManifest)
+	bool FFileAttribution::HasSameAttributes(const FFileManifest* NewFileManifest, const FFileManifest* OldFileManifest) const
 	{
 		// Currently it is not supported to rely on this, as the update process always makes new files when a file changes.
 		// This can be reconsidered when the patching process changes.
@@ -133,7 +149,7 @@ namespace BuildPatchServices
 		//     && (NewFileManifest->bIsCompressed == OldFileManifest->bIsCompressed);
 	}
 
-	void FFileAttribution::SetupFileAttributes(const FString& FilePath, const FFileManifest& FileManifest, bool bForce)
+	void FFileAttribution::SetupFileAttributes(const FString& FilePath, const FFileManifest& FileManifest, bool bForce) const
 	{
 		using namespace BuildPatchServices;
 		EAttributeFlags FileAttributes = EAttributeFlags::None;
