p4 integrate //UE5/Dev-Distro-5.5/Engine/Source/Programs/BuildPatchTool/... //UE5/Main/Engine/Source/Programs/BuildPatchTool/...

diff --git a/Engine/Source/Runtime/Online/BuildPatchServices/Private/Installer/InstallChunkSource.cpp b/Engine/Source/Runtime/Online/BuildPatchServices/Private/Installer/InstallChunkSource.cpp
--- a/Engine/Source/Runtime/Online/BuildPatchServices/Private/Installer/InstallChunkSource.cpp
+++ b/Engine/Source/Runtime/Online/BuildPatchServices/Private/Installer/InstallChunkSource.cpp
@@ -14,7 +14,8 @@
 #include "Tasks/Task.h"
 #include "BuildPatchUtil.h"
 
-DEFINE_LOG_CATEGORY_STATIC(LogInstallChunkSource, Log, All);
+DECLARE_LOG_CATEGORY_EXTERN(LogInstallChunkSource, Log, All);
+DEFINE_LOG_CATEGORY(LogInstallChunkSource);
 
 namespace BuildPatchServices
 {
@@ -29,7 +30,7 @@ namespace BuildPatchServices
 	{
 	public:
 		FInstallChunkSource(IFileSystem* FileSystem, IInstallChunkSourceStat* InInstallChunkSourceStat, const TMultiMap<FString, FBuildPatchAppManifestRef>& InInstallationSources, 
-			const TSet<FGuid>& ChunksThatWillBeNeeded, const TMap<FString, FString>& PerFileSubdirectories);
+			const TSet<FGuid>& ChunksThatWillBeNeeded, bool bInAllowLoadingAnyChunk, const TMap<FString, FString>& PerFileSubdirectories);
 		~FInstallChunkSource();
 		
 		virtual FRequestProcessFn CreateRequest(const FGuid& DataId, FMutableMemoryView DestinationBuffer, void* UserPtr, FChunkRequestCompleteDelegate CompleteFn) override;
@@ -139,6 +140,11 @@ namespace BuildPatchServices
 		// Contains the list of pieces to assemble chunks from, for each chunk.
 		TMap<FGuid, TArray<FChunkSourceDetails>> ChunkSources;
 
+		// When we try and recreate chunks from an existing installation, there are chunks that don't get enumerated
+		// as producible. Additionally, the chunks might not fill the full window. Setting this allows any request for
+		// a chunk where we have any parts to succeed and be zero filled for any missing parts.
+		bool bAllowLoadingAnyChunk = false;
+
 		// Contains files that have a custom subdirectory to install into. This is accessed from multiple threads
 		// but should be static across execution and so requires no lock.
 		const TMap<FString, FString>& PerFileSubdirectories;
@@ -146,9 +152,10 @@ namespace BuildPatchServices
 
 	FInstallChunkSource::FInstallChunkSource(IFileSystem* InFileSystem, IInstallChunkSourceStat* InInstallChunkSourceStat, 
 		const TMultiMap<FString, FBuildPatchAppManifestRef>& InInstallationSources, const TSet<FGuid>& ChunksThatWillBeNeeded,
-		const TMap<FString, FString>& InPerFileSubdirectories)
+		bool bInAllowLoadingAnyChunk, const TMap<FString, FString>& InPerFileSubdirectories)
 		: FileSystem(InFileSystem)
 		, InstallChunkSourceStat(InInstallChunkSourceStat)
+		, bAllowLoadingAnyChunk(bInAllowLoadingAnyChunk)
 		, PerFileSubdirectories(InPerFileSubdirectories)
 	{
 		TRACE_CPUPROFILER_EVENT_SCOPE(InstallChunkSource_ctor);
@@ -242,7 +249,7 @@ namespace BuildPatchServices
 						}
 
 						if (ChunkPart.Offset == 0 &&
-							ChunkInfoPtr->WindowSize == ChunkPart.Size)
+							ChunkInfoPtr->DataSizeUncompressed == ChunkPart.Size)
 						{
 							// Got the whole thing, don't need to fill out the chunk parts.
 							FChunkSourceDetails NewInfo;
@@ -283,7 +290,7 @@ namespace BuildPatchServices
 					const FChunkInfo* ChunkInfoPtr = InstalledManifest->GetChunkInfo(ChunkPart.Guid);
 
 					if (PartialStructure.GetHead() == PartialStructure.GetTail() &&
-						PartialStructure.GetHead()->GetSize() == ChunkInfoPtr->WindowSize)
+						PartialStructure.GetHead()->GetSize() == ChunkInfoPtr->DataSizeUncompressed)
 					{
 						// Done with this chunk, free the structures space.
 						ChunkAssemblyInfo.bComplete = true;
@@ -304,7 +311,11 @@ namespace BuildPatchServices
 				if (!ChunkAssemblyPair.Value.bComplete)
 				{
 					IncompleteChunks++;
-					continue;
+					// Skip incomplete chunks unless we allow loading incomplete chunks
+					if (!bAllowLoadingAnyChunk)
+					{
+						continue;
+					}
 				}
 
 				TArray<FChunkSourceDetails>& ChunkSource = ChunkSources.FindOrAdd(ChunkAssemblyPair.Key);
@@ -315,7 +326,10 @@ namespace BuildPatchServices
 				}
 
 				ChunkSource = MoveTemp(ChunkAssemblyPair.Value.SourceDetails);
-				CompleteChunks++;
+				if (ChunkAssemblyPair.Value.bComplete)
+				{
+					CompleteChunks++;
+				}
 
 				AvailableInBuilds.Add(ChunkAssemblyPair.Key);
 			}
@@ -327,7 +341,8 @@ namespace BuildPatchServices
 			}
 
 		}
-		UE_LOG(LogInstallChunkSource, Log, TEXT("Useful Sources:%d. Available Chunks:%d."), InstallationSources.Num(), AvailableInBuilds.Num());
+
+		UE_LOG(LogInstallChunkSource, Log, TEXT("Useful Sources:%d/%d. Available Chunks:%d/%d."), InstallationSources.Num(), InInstallationSources.Num(), AvailableInBuilds.Num(), ChunksThatWillBeNeeded.Num());
 	}
 
 	FInstallChunkSource::~FInstallChunkSource()
@@ -404,6 +419,14 @@ namespace BuildPatchServices
 			const FChunkInfo* ChunkInfoPtr = FoundInstallManifest->GetChunkInfo(DataId);
 			if (!ChunkInfoPtr || !ChunkSource)
 			{
+				if (!ChunkSource)
+				{
+					UE_LOG(LogInstallChunkSource, Error, TEXT("Chunk %s not found in sources."), *DataId.ToString());
+				}
+				else
+				{
+					UE_LOG(LogInstallChunkSource, Error, TEXT("Chunk info for %s not found in install manifest."), *DataId.ToString());
+				}
 				ActivityRecord.CyclesEnd = FStatsCollector::GetCycles();
 				InstallChunkSourceStat->OnLoadComplete(DataId, IInstallChunkSourceStat::ELoadResult::MissingPartInfo, ActivityRecord);
 
@@ -500,7 +523,26 @@ namespace BuildPatchServices
 				// The expectation is that we get the full chunk only once we've assembled all of the parts
 				// provided by the manifest, so the last iteration should set this to true. If it isn't the last
 				// iteration, then we'll hit the faux-error case at the top of the loop.
-				bLoadedWholeChunk = ChunkBlocks.GetHead() && ChunkBlocks.GetHead() == ChunkBlocks.GetTail() && ChunkBlocks.GetHead()->GetSize() == ChunkInfoPtr->WindowSize;
+				bLoadedWholeChunk = ChunkBlocks.GetHead() && ChunkBlocks.GetHead() == ChunkBlocks.GetTail() && ChunkBlocks.GetHead()->GetSize() == ChunkInfoPtr->DataSizeUncompressed;
+			}
+
+			// If we didn't have enough parts but we're allowing loading partial chunks, succeed the chunk
+			// and zero fill what we didn't have.
+			if (!bLoadedWholeChunk && 
+				Result == IInstallChunkSourceStat::ELoadResult::Success &&
+				bAllowLoadingAnyChunk)
+			{
+				FBlockStructure Remnants;
+				Remnants.Add(0, ChunkInfoPtr->DataSizeUncompressed);
+				Remnants.Remove(ChunkBlocks);
+
+				const FBlockEntry* Entry = Remnants.GetHead();
+				while (Entry)
+				{
+					FMemory::Memset((uint8*)DestinationBuffer.GetData() + Entry->GetOffset(), 0, Entry->GetSize());
+					Entry = Entry->GetNext();
+				}
+				bLoadedWholeChunk = true;
 			}
 			
 			if (!bLoadedWholeChunk)
@@ -616,9 +658,9 @@ namespace BuildPatchServices
 	}
 
 	IConstructorInstallChunkSource* IConstructorInstallChunkSource::CreateInstallSource(IFileSystem* FileSystem, IInstallChunkSourceStat* InstallChunkSourceStat, 
-		const TMultiMap<FString, FBuildPatchAppManifestRef>& InstallationSources, const TSet<FGuid>& ChunksThatWillBeNeeded, const TMap<FString, FString>& PerFileSubdirectories)
+		const TMultiMap<FString, FBuildPatchAppManifestRef>& InstallationSources, const TSet<FGuid>& ChunksThatWillBeNeeded, bool bAllowLoadingAnyChunk, const TMap<FString, FString>& PerFileSubdirectories)
 	{
-		return new FInstallChunkSource(FileSystem, InstallChunkSourceStat, InstallationSources, ChunksThatWillBeNeeded, PerFileSubdirectories);
+		return new FInstallChunkSource(FileSystem, InstallChunkSourceStat, InstallationSources, ChunksThatWillBeNeeded, bAllowLoadingAnyChunk, PerFileSubdirectories);
 	}
 
 	const TCHAR* ToString(const IInstallChunkSourceStat::ELoadResult& LoadResult)
diff --git a/Engine/Source/Runtime/Online/BuildPatchServices/Private/Installer/InstallChunkSource.h b/Engine/Source/Runtime/Online/BuildPatchServices/Private/Installer/InstallChunkSource.h
--- a/Engine/Source/Runtime/Online/BuildPatchServices/Private/Installer/InstallChunkSource.h
+++ b/Engine/Source/Runtime/Online/BuildPatchServices/Private/Installer/InstallChunkSource.h
@@ -51,7 +51,8 @@ namespace BuildPatchServices
 			IFileSystem* FileSystem, 
 			IInstallChunkSourceStat* InstallChunkSourceStat, 
 			const TMultiMap<FString, FBuildPatchAppManifestRef>& InstallationSources,
-			const TSet<FGuid>& ChunksThatWillBeNeeded,
+			const TSet<FGuid>& ChunksThatWillBeNeeded, 
+			bool bAllowLoadingAnyChunk,
 			const TMap<FString, FString>& InPerFileSubdirectories);
 	};
 
diff --git a/Engine/Source/Runtime/Online/BuildPatchServices/Private/Installer/InstallerError.cpp b/Engine/Source/Runtime/Online/BuildPatchServices/Private/Installer/InstallerError.cpp
--- a/Engine/Source/Runtime/Online/BuildPatchServices/Private/Installer/InstallerError.cpp
+++ b/Engine/Source/Runtime/Online/BuildPatchServices/Private/Installer/InstallerError.cpp
@@ -30,6 +30,8 @@ namespace BuildPatchServices
 				return NSLOCTEXT("BuildPatchInstallError", "BuildPatchInstallShortError_UserCanceled", "User cancelled.");
 			case EBuildPatchInstallError::PrerequisiteError: 
 				return NSLOCTEXT("BuildPatchInstallError", "BuildPatchInstallShortError_PrerequisiteError", "The necessary prerequisites have failed to install. Please contact support.");
+			case EBuildPatchInstallError::CustomUninstallActionError: 
+				return NSLOCTEXT("BuildPatchInstallError", "BuildPatchInstallShortError_UninstallCustomActionError", "The application's uninstall action has failed to execute. Please contact support.");
 			case EBuildPatchInstallError::InitializationError: 
 				return NSLOCTEXT("BuildPatchInstallError", "BuildPatchInstallShortError_InitializationError", "The installer failed to initialize. Please contact support.");
 			case EBuildPatchInstallError::PathLengthExceeded: 
@@ -156,11 +158,11 @@ namespace BuildPatchServices
 			RegisteredDelegates.GenerateValueArray(DelegatesToCall);
 			if (IsCancelled())
 			{
-				UE_LOG(LogBuildPatchServices, Log, TEXT("%s %s"), *EnumToString(ErrorType), *ErrorCode);
+				UE_LOG(LogBuildPatchServices, Log, TEXT("EBuildPatchInstallError::%s %s"), LexToString(ErrorType), *ErrorCode);
 			}
 			else
 			{
-				UE_LOG(LogBuildPatchServices, Error, TEXT("%s %s"), *EnumToString(ErrorType), *ErrorCode);
+				UE_LOG(LogBuildPatchServices, Error, TEXT("EBuildPatchInstallError::%s %s"), LexToString(ErrorType), *ErrorCode);
 			}
 		}
 		ThreadLockCs.Unlock();
@@ -196,25 +198,4 @@ namespace BuildPatchServices
 	{
 		return new FInstallerError();
 	}
-
-	FString EnumToString(const EBuildPatchInstallError& InErrorType)
-	{
-		switch (InErrorType)
-		{
-			case EBuildPatchInstallError::NoError: return TEXT("SUCCESS");
-			case EBuildPatchInstallError::DownloadError: return TEXT("EBuildPatchInstallError::DownloadError");
-			case EBuildPatchInstallError::FileConstructionFail: return TEXT("EBuildPatchInstallError::FileConstructionFail");
-			case EBuildPatchInstallError::MoveFileToInstall: return TEXT("EBuildPatchInstallError::MoveFileToInstall");
-			case EBuildPatchInstallError::BuildVerifyFail: return TEXT("EBuildPatchInstallError::BuildVerifyFail");
-			case EBuildPatchInstallError::ApplicationClosing: return TEXT("EBuildPatchInstallError::ApplicationClosing");
-			case EBuildPatchInstallError::ApplicationError: return TEXT("EBuildPatchInstallError::ApplicationError");
-			case EBuildPatchInstallError::UserCanceled: return TEXT("EBuildPatchInstallError::UserCanceled");
-			case EBuildPatchInstallError::PrerequisiteError: return TEXT("EBuildPatchInstallError::PrerequisiteError");
-			case EBuildPatchInstallError::InitializationError: return TEXT("EBuildPatchInstallError::InitializationError");
-			case EBuildPatchInstallError::PathLengthExceeded: return TEXT("EBuildPatchInstallError::PathLengthExceeded");
-			case EBuildPatchInstallError::OutOfDiskSpace: return TEXT("EBuildPatchInstallError::OutOfDiskSpace");
-			default: return TEXT("EBuildPatchInstallError::InvalidOrMax");
-		}
-	}
-
 }
diff --git a/Engine/Source/Runtime/Online/BuildPatchServices/Private/Installer/InstallerError.h b/Engine/Source/Runtime/Online/BuildPatchServices/Private/Installer/InstallerError.h
--- a/Engine/Source/Runtime/Online/BuildPatchServices/Private/Installer/InstallerError.h
+++ b/Engine/Source/Runtime/Online/BuildPatchServices/Private/Installer/InstallerError.h
@@ -19,8 +19,12 @@ namespace BuildPatchServices
 		static const TCHAR* ChunkReferenceTracking = TEXT("05");
 		static const TCHAR* MissingPrereqForPrereqOnlyInstall = TEXT("06");
 		static const TCHAR* InvalidDataInManifest = TEXT("07");
-		static const TCHAR* InvalidInstallToMemory = TEXT("08");
-		static const TCHAR* InvalidMemoryConstructionConfig = TEXT("09");
+		static const TCHAR* EmptyCloudDirectories = TEXT("08");
+		static const TCHAR* VerifyOnlyActionMixedWithOtherActions = TEXT("09");
+		static const TCHAR* MissingEncryptionKeyForNonPreload = TEXT("10");
+		static const TCHAR* InvalidEncryptionKeyForNonPreload = TEXT("11");
+		static const TCHAR* InvalidInstallToMemory = TEXT("11");
+		static const TCHAR* InvalidMemoryConstructionConfig = TEXT("12");
 	}
 
 	/**
@@ -66,7 +70,9 @@ namespace BuildPatchServices
 		static const TCHAR* SerializationError = TEXT("06");
 		static const TCHAR* TrackingError = TEXT("07");
 		static const TCHAR* InternalConsistencyFailure = TEXT("08");
-		static const TCHAR* CloseFileError = TEXT("09");
+		static const TCHAR* SymlinkCreateFail = TEXT("09");
+		static const TCHAR* SymlinkCreateForbidden = TEXT("10");
+		static const TCHAR* CloseFileError = TEXT("11");
 	}
 
 	/**
@@ -84,12 +90,12 @@ namespace BuildPatchServices
 	 */
 	namespace VerifyErrorCodes
 	{
-		static const TCHAR* UnknownFail = TEXT("00");
-		static const TCHAR* FinalCheck = TEXT("01");
+		static const TCHAR* UnknownFail = TEXT("01");
 		static const TCHAR* FileMissing = TEXT("02");
 		static const TCHAR* OpenFileFailed = TEXT("03");
 		static const TCHAR* HashCheckFailed = TEXT("04");
 		static const TCHAR* FileSizeFailed = TEXT("05");
+		static const TCHAR* DeleteCorruptedFileFailed = TEXT("06");
 	}
 
 	/**
@@ -109,7 +115,16 @@ namespace BuildPatchServices
 	}
 
 	/**
-	 * A list of the error codes we will use for each case of cancellation.
+	 * A list of the error codes we will use for each case of the uninstall action failure.
+	 */
+	namespace UninstallCustomActionErrorPrefixes
+	{
+		static const TCHAR* ExecuteCode = TEXT("E");
+		static const TCHAR* ReturnCode = TEXT("R");
+	}
+
+	/**
+	 * A list of the error codes we will use for each case of the prerequisites install failure.
 	 */
 	namespace PrerequisiteErrorPrefixes
 	{
@@ -223,11 +238,4 @@ namespace BuildPatchServices
 		 */
 		static IInstallerError* Create();
 	};
-	
-	/**
-	 * Returns the string representation of the specified FBuildPatchInstallError value. Used for logging only.
-	 * @param ErrorType     The error type value.
-	 * @return A string value for the EBuildPatchInstallError enum.
-	 */
-	FString EnumToString(const EBuildPatchInstallError& ErrorType);
 }
diff --git a/Engine/Source/Runtime/Online/BuildPatchServices/Private/Installer/MachineConfig.cpp b/Engine/Source/Runtime/Online/BuildPatchServices/Private/Installer/MachineConfig.cpp
--- a/Engine/Source/Runtime/Online/BuildPatchServices/Private/Installer/MachineConfig.cpp
+++ b/Engine/Source/Runtime/Online/BuildPatchServices/Private/Installer/MachineConfig.cpp
@@ -53,7 +53,11 @@ namespace BuildPatchServices
 	FMachineConfig::FMachineConfig(const FString& InLocalMachineConfigFile, bool bInAlwaysFlushChanges)
 		: LocalMachineConfigFile(InLocalMachineConfigFile)
 		, bAlwaysFlushChanges(bInAlwaysFlushChanges)
-	{ }
+	{
+		// Explicitly add the config file to the cache so that it's not "found" later.
+		// "Found" files are loaded in read-only mode ("NoSave" flag is "true"), but we may need to write new prereq IDs to this file.
+		GConfig->Add(LocalMachineConfigFile, FConfigFile());
+	}
 
 	TSet<FString> FMachineConfig::LoadInstalledPrereqIds()
 	{
diff --git a/Engine/Source/Runtime/Online/BuildPatchServices/Private/Installer/MessagePump.cpp b/Engine/Source/Runtime/Online/BuildPatchServices/Private/Installer/MessagePump.cpp
--- a/Engine/Source/Runtime/Online/BuildPatchServices/Private/Installer/MessagePump.cpp
+++ b/Engine/Source/Runtime/Online/BuildPatchServices/Private/Installer/MessagePump.cpp
@@ -18,6 +18,13 @@ namespace BuildPatchServices
 
 	// Queue of request responses of chunk we are looking for
 	typedef TQueue<FChunkUriRequestResponse, EQueueMode::Mpsc> FChunkUriQueue;
+
+	// Tuple of the request and response for finding a delta uri location
+	typedef TTuple<FDeltaUriRequest, TFunction<void(FExpiratedUrisResponse)>> FDeltaUriRequestResponse;
+
+	// Queue of request responses of delta file we are looking for
+	typedef TQueue<FDeltaUriRequestResponse, EQueueMode::Mpsc> FDeltaUriQueue;
+
 }
 
 namespace MessagePumpHelpers
@@ -48,14 +55,20 @@ namespace BuildPatchServices
 		virtual void SendMessage(FInstallationFileAction Message) override;
 		virtual void SendMessage(FGenericMessage Message) override;
 		virtual bool SendRequest(FChunkUriRequest Request, TFunction<void(FChunkUriResponse)> OnResponse) override;
+		virtual void SendRequest(FDeltaUriRequest Request, TFunction<void(FExpiratedUrisResponse)> OnResponse) override;
+		virtual void PumpMessages(const TArray<FMessageHandler*>& Handlers) override;
 		virtual void PumpMessages() override;
 		virtual void RegisterMessageHandler(FMessageHandler* MessageHandler) override;
 		virtual void UnregisterMessageHandler(FMessageHandler* MessageHandler) override;
 		// IMessagePump interface end.
 
 	private:
+		template<typename FResponse, typename FRequest>
+		void PumpMessages(const TArray<FMessageHandler*>& InHandlers, FRequest& RequestQueue);
+
 		FMessageQueue MessageQueue;
 		FChunkUriQueue ChunkUriRequestQueue;
+		FDeltaUriQueue DeltaUriRequestQueue;
 		FDefaultMessageHandler DefaultMessageHandler;
 		TArray<FMessageHandler*> Handlers;
 		// std::atomic will not expose integral operations with enum types.
@@ -87,42 +100,62 @@ namespace BuildPatchServices
 		return false;
 	}
 
+	template<typename FResponse, typename FRequest>
+	void FMessagePump::PumpMessages(const TArray<FMessageHandler*>& InHandlers, FRequest& RequestQueue)
+	{
+		FResponse ChunkUriRequestResponse;
+		while (RequestQueue.Dequeue(ChunkUriRequestResponse))
+		{
+			bool bHandled = false;
+			for (BuildPatchServices::FMessageHandler* Handler : InHandlers)
+			{
+				if (Handler->HandleRequest(ChunkUriRequestResponse.template Get<0>(), ChunkUriRequestResponse.template Get<1>()))
+				{
+					bHandled = true;
+					break;
+				}
+			}
+			if (!bHandled)
+			{
+				DefaultMessageHandler.HandleRequest(ChunkUriRequestResponse.template Get<0>(), ChunkUriRequestResponse.template Get<1>());
+			} 
+		}
+	}
+
 	void FMessagePump::PumpMessages()
+	{
+		PumpMessages(Handlers);
+	}
+
+	void FMessagePump::PumpMessages(const TArray<FMessageHandler*>& ExtHandlers)
 	{
 		FMessageUnion MessageUnion;
 		while (MessageQueue.Dequeue(MessageUnion))
 		{
-			if (MessagePumpHelpers::TryPump<FChunkSourceEvent>(Handlers, MessageUnion))
+			if (MessagePumpHelpers::TryPump<FChunkSourceEvent>(ExtHandlers, MessageUnion))
 			{
 				continue;
 			}
-			else if (MessagePumpHelpers::TryPump<FInstallationFileAction>(Handlers, MessageUnion))
+			else if (MessagePumpHelpers::TryPump<FInstallationFileAction>(ExtHandlers, MessageUnion))
 			{
 				continue;
 			}
-			else if (MessagePumpHelpers::TryPump<FGenericMessage>(Handlers, MessageUnion))
+			else if (MessagePumpHelpers::TryPump<FGenericMessage>(ExtHandlers, MessageUnion))
 			{
 				continue;
 			}
 		}
 
-		FChunkUriRequestResponse ChunkUriRequestResponse;
-		while (ChunkUriRequestQueue.Dequeue(ChunkUriRequestResponse))
-		{
-			bool bHandled = false;
-			for (BuildPatchServices::FMessageHandler* Handler : Handlers)
-			{
-				if (Handler->HandleRequest(ChunkUriRequestResponse.Get<0>(), ChunkUriRequestResponse.Get<1>()))
-				{
-					bHandled = true;
-					break;
-				}
-			}
-			if (!bHandled)
-			{
-				DefaultMessageHandler.HandleRequest(ChunkUriRequestResponse.Get<0>(), ChunkUriRequestResponse.Get<1>());
-			}
-		}
+		// Handle get chunk URI request
+		PumpMessages<FChunkUriRequestResponse>(ExtHandlers, ChunkUriRequestQueue);
+
+		// Handle get delta file URI request
+		PumpMessages<FDeltaUriRequestResponse>(ExtHandlers, DeltaUriRequestQueue);
+	}
+
+	void FMessagePump::SendRequest(FDeltaUriRequest Request, TFunction<void(FExpiratedUrisResponse)> OnResponse)
+	{
+		DeltaUriRequestQueue.Enqueue(FDeltaUriRequestResponse(MoveTemp(Request), MoveTemp(OnResponse)));
 	}
 
 	void FMessagePump::RegisterMessageHandler(FMessageHandler* InHandler)
diff --git a/Engine/Source/Runtime/Online/BuildPatchServices/Private/Installer/MessagePump.h b/Engine/Source/Runtime/Online/BuildPatchServices/Private/Installer/MessagePump.h
--- a/Engine/Source/Runtime/Online/BuildPatchServices/Private/Installer/MessagePump.h
+++ b/Engine/Source/Runtime/Online/BuildPatchServices/Private/Installer/MessagePump.h
@@ -53,10 +53,19 @@ namespace BuildPatchServices
 		 */
 		virtual bool SendRequest(FChunkUriRequest Request, TFunction<void(FChunkUriResponse)> OnResponse) = 0;
 
+		/**
+		 * Sends out a request to resolve the uri to the delta file location
+		 * @param Request   Request for the delta file location
+		 * @param OnResponse   A delegate to call with a response to delta location
+		 */
+		virtual void SendRequest(FDeltaUriRequest Request, TFunction<void(FExpiratedUrisResponse)> OnResponse) = 0;
+
 		/**
 		 * Dequeues received messages, pushing them to the provided handlers.
 		 * NOTE: PumpMessages, RegisterMessageHandler, and UnregisterMessageHandler MUST all be called from the same thread.
+		 * @param Handlers      The array of handlers.
 		 */
+		virtual void PumpMessages(const TArray<FMessageHandler*>& Handlers) = 0;
 		virtual void PumpMessages() = 0;
 
 		/**
diff --git a/Engine/Source/Runtime/Online/BuildPatchServices/Private/Installer/OptimisedDelta.cpp b/Engine/Source/Runtime/Online/BuildPatchServices/Private/Installer/OptimisedDelta.cpp
--- a/Engine/Source/Runtime/Online/BuildPatchServices/Private/Installer/OptimisedDelta.cpp
+++ b/Engine/Source/Runtime/Online/BuildPatchServices/Private/Installer/OptimisedDelta.cpp
@@ -11,72 +11,124 @@
 #include "BuildPatchUtil.h"
 #include "BuildPatchMergeManifests.h"
 
-DECLARE_LOG_CATEGORY_CLASS(LogOptimisedDelta, Log, All);
+#include <atomic>
+
+DECLARE_LOG_CATEGORY_EXTERN(LogOptimisedDelta, Log, All);
+DEFINE_LOG_CATEGORY(LogOptimisedDelta);
 
 namespace ConfigHelpers
 {
-	int32 LoadDeltaRetries(int32 Min)
+	int32 LoadRetries(const int32 Min, const FString& PreferenceName)
 	{
-		int32 DeltaRetries = 6;
-		GConfig->GetInt(TEXT("Portal.BuildPatch"), TEXT("DeltaRetries"), DeltaRetries, GEngineIni);
+		int32 DeltaRetries = 5;
+		GConfig->GetInt(TEXT("Portal.BuildPatch"), *PreferenceName, DeltaRetries, GEngineIni);
 		DeltaRetries = FMath::Clamp<int32>(DeltaRetries, Min, 1000);
 		return DeltaRetries;
 	}
+
+	int32 LoadDeltaRetries(const int32 Min)
+	{
+		return LoadRetries(Min, TEXT("DeltaRetries"));
+	}
+
+	int32 LoadDeltaUriRetries(const int32 Min)
+	{
+		return LoadRetries(Min, TEXT("DeltaUriRetries"));
+	}
 }
 
 namespace BuildPatchServices
 {
-	class FOptimisedDelta
-		: public IOptimisedDelta
+	template<typename ResultType>
+	class TAbortablePromise
+	{
+	public:
+		TAbortablePromise() = default;
+		~TAbortablePromise()
+		{
+			Abort();
+		}
+
+		void Abort()
+		{
+			if (!bSetValue)
+			{
+				SetValue(MakeShared<ResultType, ESPMode::ThreadSafe>(MakeError(UserCancelErrorCodes::UserRequested)));
+			}
+		}
+
+		TFuture<TSharedPtr<ResultType, ESPMode::ThreadSafe>> GetFuture()
+		{
+			return Promise.GetFuture();
+		}
+
+		void SetValue(TSharedPtr<ResultType, ESPMode::ThreadSafe> Result)
+		{
+			bSetValue = true;
+			Promise.SetValue(MoveTemp(Result));
+		}
+
+	private:
+		std::atomic<bool> bSetValue{ false };
+		TPromise<TSharedPtr<ResultType, ESPMode::ThreadSafe>> Promise;
+	};
+
+	class FOptimisedDelta :
+		public IOptimisedDelta,
+		public TSharedFromThis<FOptimisedDelta>
 	{
 	public:
-		FOptimisedDelta(const FOptimisedDeltaConfiguration& Configuration, FOptimisedDeltaDependencies&& Dependencies);
+		FOptimisedDelta(const FOptimisedDeltaConfiguration& Configuration, FOptimisedDeltaDependencies Dependencies);
 
 		// IOptimisedDelta interface begin.
 		virtual const IOptimisedDelta::FResultValueOrError& GetResult() const override;
 		virtual int32 GetMetaDownloadSize() const override;
+		virtual bool IsReady() const override;
 		// IOptimisedDelta interface end.
 
+		void Initialize();
+
 	private:
 		bool RequiresChunkDownload();
 		void OnDownloadComplete(int32 RequestId, const FDownloadRef& Download);
 		bool ShouldRetry(const FDownloadRef& Download);
+		void RequestDeltaFileUri();
+		void DeltaUriCompletionDelegate(bool bSuccessful, const FString& DeltaUri, bool bNeedRetry);
 		void SetResult(const FBuildPatchAppManifestPtr& ResultManifest);
 		void SetFailedDownload();
 
 	private:
 		const FOptimisedDeltaConfiguration Configuration;
-		IDownloadService* const DownloadService;
+		const FOptimisedDeltaDependencies Dependencies;
 		const FString RelativeDeltaFilePath;
 		const int32 DeltaRetries;
-		TUniqueFunction<void(const IOptimisedDelta::FResultValueOrError&)> OnComplete;
-		EDeltaPolicy DeltaPolicy;
-		int32 CloudDirIdx;
-		int32 RetryCount;
+		const int32 DeltaUriRetries;
+		EDeltaPolicy DeltaPolicy = Configuration.DeltaPolicy;
+		int32 CloudDirIdx = 0;
+		int32 RetryCount = 0;
+		int32 DeltaUriRetryCount = 0;
 		FDownloadProgressDelegate ChunkDeltaProgress;
 		FDownloadCompleteDelegate ChunkDeltaComplete;
+		// We use TAbortablePromise in the case when we destroy the object before actual delta is received
 		// TPromise and TFuture are restricted to types that are copyable, assignable, and default constructible. Value should not be exposed.
-		TPromise<TSharedPtr<FResultValueOrError>> ChunkDeltaPromise;
-		TFuture<TSharedPtr<FResultValueOrError>> ChunkDeltaFuture;
-		FThreadSafeCounter DownloadedBytes;
-		const TCHAR* ErrorCode;
+		TAbortablePromise<IOptimisedDelta::FResultValueOrError> ChunkDeltaPromise;
+		TFuture<TSharedPtr<FResultValueOrError, ESPMode::ThreadSafe>> ChunkDeltaFuture{ ChunkDeltaPromise.GetFuture() };
+		FThreadSafeCounter DownloadedBytes = 0;
+		const TCHAR* ErrorCode = nullptr;
 	};
 
-	FOptimisedDelta::FOptimisedDelta(const FOptimisedDeltaConfiguration& InConfiguration, FOptimisedDeltaDependencies&& InDependencies)
+	FOptimisedDelta::FOptimisedDelta(const FOptimisedDeltaConfiguration& InConfiguration, FOptimisedDeltaDependencies InDependencies)
 		: Configuration(InConfiguration)
-		, DownloadService(InDependencies.DownloadService)
+		, Dependencies(MoveTemp(InDependencies))
 		, RelativeDeltaFilePath(Configuration.SourceManifest.IsValid() ? FBuildPatchUtils::GetChunkDeltaFilename(*Configuration.SourceManifest.Get(), Configuration.DestinationManifest.Get(), Configuration.DeltaFilenameTrailer) : TEXT(""))
-		, DeltaRetries(ConfigHelpers::LoadDeltaRetries(Configuration.CloudDirectories.Num()))
-		, OnComplete(MoveTemp(InDependencies.OnComplete))
-		, DeltaPolicy(Configuration.DeltaPolicy)
-		, CloudDirIdx(0)
-		, RetryCount(0)
-		, ChunkDeltaComplete(FDownloadCompleteDelegate::CreateRaw(this, &FOptimisedDelta::OnDownloadComplete))
-		, ChunkDeltaPromise()
-		, ChunkDeltaFuture(ChunkDeltaPromise.GetFuture())
-		, DownloadedBytes(0)
-		, ErrorCode(nullptr)
+		, DeltaRetries(ConfigHelpers::LoadDeltaRetries(Configuration.RetriesNumber))
+		, DeltaUriRetries(ConfigHelpers::LoadDeltaUriRetries(Configuration.RetriesNumber))
+	{}
+
+	void FOptimisedDelta::Initialize()
 	{
+		ChunkDeltaComplete = FDownloadCompleteDelegate::CreateSP(this, &FOptimisedDelta::OnDownloadComplete);
+
 		// There are some conditions in which we do not use a delta.
 		const bool bNoSourceManifest = Configuration.SourceManifest.IsValid() == false;
 		const bool bNotPatching = Configuration.InstallMode == EInstallMode::PrereqOnly;
@@ -89,8 +141,8 @@ namespace BuildPatchServices
 		// Kick off the request if we should be.
 		if (DeltaPolicy != EDeltaPolicy::Skip)
 		{
-			UE_LOG(LogOptimisedDelta, Display, TEXT("Requesting optimised delta file %s"), *(Configuration.CloudDirectories[CloudDirIdx] / RelativeDeltaFilePath));
-			DownloadService->RequestFile(Configuration.CloudDirectories[CloudDirIdx] / RelativeDeltaFilePath, ChunkDeltaComplete, ChunkDeltaProgress);
+			UE_LOG(LogOptimisedDelta, Log, TEXT("Requesting optimised delta file %s"), *RelativeDeltaFilePath);
+			RequestDeltaFileUri();
 		}
 		// Otherwise we provide the standard destination manifest.
 		else
@@ -101,7 +153,7 @@ namespace BuildPatchServices
 
 	const IOptimisedDelta::FResultValueOrError& FOptimisedDelta::GetResult() const
 	{
-		TSharedPtr<FResultValueOrError> Result = ChunkDeltaFuture.Get();
+		TSharedPtr<FResultValueOrError, ESPMode::ThreadSafe> Result = ChunkDeltaFuture.Get();
 		// Result should never be set nullptr.
 		check(Result.IsValid());
 		return *Result;
@@ -113,6 +165,11 @@ namespace BuildPatchServices
 		return DownloadedBytes.GetValue();
 	}
 
+	bool FOptimisedDelta::IsReady() const
+	{
+		return ChunkDeltaFuture.IsReady();
+	}
+
 	bool FOptimisedDelta::RequiresChunkDownload()
 	{
 		if (Configuration.SourceManifest.IsValid())
@@ -123,7 +180,7 @@ namespace BuildPatchServices
 				const FFileManifest* OldFile = Configuration.SourceManifest->GetFileManifest(DestinationFile);
 				const FFileManifest* NewFile = Configuration.DestinationManifest->GetFileManifest(DestinationFile);
 				// Different hash means different file, missing OldFile means added file.
-				if (!OldFile || OldFile->FileHash != NewFile->FileHash)
+				if (!OldFile || OldFile->SHA1Hash != NewFile->SHA1Hash)
 				{
 					return true;
 				}
@@ -137,98 +194,117 @@ namespace BuildPatchServices
 
 	void FOptimisedDelta::OnDownloadComplete(int32 RequestId, const FDownloadRef& Download)
 	{
-		UE_LOG(LogOptimisedDelta, Log, TEXT("Download complete for %s, launching comlpletion task"), *RelativeDeltaFilePath);
-
-		// Kick a task to finish this so as not to block the callback thread
-		UE::Tasks::Launch(UE_SOURCE_LOCATION, [this, RequestId, Download] 
+		if (Download->ResponseSuccessful())
 		{
-			if (Download->ResponseSuccessful())
+			// Perform a merge with current manifest so that the delta can support missing out unnecessary information.
+			FBuildPatchAppManifestPtr NewManifest;
+			FBuildPatchAppManifestRef DeltaManifest = MakeShareable(new FBuildPatchAppManifest());
+			if (DeltaManifest->DeserializeFromData(Download->GetData()))
 			{
-				// Perform a merge with current manifest so that the delta can support missing out unnecessary information.
-				FBuildPatchAppManifestPtr NewManifest;
-				FBuildPatchAppManifestRef DeltaManifest = MakeShareable(new FBuildPatchAppManifest());
-				if (DeltaManifest->DeserializeFromData(Download->GetData()))
+				// Verify each file received matches what we expect.
+				TSet<FString> DeltaFilenameList;
+				DeltaManifest->GetFileList(DeltaFilenameList);
+				bool bDeltaAccepted = true;
+				for (const FString& DeltaFilename : DeltaFilenameList)
 				{
-					// Verify each file received matches what we expect.
-					TSet<FString> DeltaFilenameList;
-					DeltaManifest->GetFileList(DeltaFilenameList);
-					bool bDeltaAccepted = true;
-					for (const FString& DeltaFilename : DeltaFilenameList)
+					const FFileManifest* OriginalFileManifest = Configuration.DestinationManifest->GetFileManifest(DeltaFilename);
+					const FFileManifest* DeltaFileManifest = DeltaManifest->GetFileManifest(DeltaFilename);
+					if (OriginalFileManifest != nullptr && DeltaFileManifest != nullptr)
 					{
-						const FFileManifest* OriginalFileManifest = Configuration.DestinationManifest->GetFileManifest(DeltaFilename);
-						const FFileManifest* DeltaFileManifest = DeltaManifest->GetFileManifest(DeltaFilename);
-						if (OriginalFileManifest != nullptr && DeltaFileManifest != nullptr)
+						FSHAHash OriginalFileHash;
+						FSHAHash DeltaFileHash;
+						Configuration.DestinationManifest->GetFileHash(DeltaFilename, OriginalFileHash);
+						DeltaManifest->GetFileHash(DeltaFilename, DeltaFileHash);
+						if (OriginalFileHash != DeltaFileHash)
 						{
-							FSHAHash OriginalFileHash;
-							FSHAHash DeltaFileHash;
-							Configuration.DestinationManifest->GetFileHash(DeltaFilename, OriginalFileHash);
-							DeltaManifest->GetFileHash(DeltaFilename, DeltaFileHash);
-							if (OriginalFileHash != DeltaFileHash)
-							{
-								UE_LOG(LogOptimisedDelta, Log, TEXT("Rejected optimised delta due to file SHA1 mismatch."));
-								ErrorCode = DownloadErrorCodes::RejectedDeltaFile;
-								bDeltaAccepted = false;
-								break;
-							}
-						}
-						else
-						{
-							UE_LOG(LogOptimisedDelta, Log, TEXT("Rejected optimised delta due to file list mismatch."));
+							UE_LOG(LogOptimisedDelta, Log, TEXT("Rejected optimised delta due to file SHA1 mismatch."));
 							ErrorCode = DownloadErrorCodes::RejectedDeltaFile;
 							bDeltaAccepted = false;
 							break;
 						}
 					}
-					// Only use if accepted so far.
-					if (bDeltaAccepted)
+					else
 					{
-						NewManifest = FBuildMergeManifests::MergeDeltaManifest(Configuration.DestinationManifest, DeltaManifest);
-						if (!NewManifest.IsValid())
-						{
-							UE_LOG(LogOptimisedDelta, Log, TEXT("Rejected optimised delta due to merge failure."));
-							ErrorCode = DownloadErrorCodes::RejectedDeltaFile;
-							bDeltaAccepted = false;
-						}
+						UE_LOG(LogOptimisedDelta, Log, TEXT("Rejected optimised delta due to file list mismatch."));
+						ErrorCode = DownloadErrorCodes::RejectedDeltaFile;
+						bDeltaAccepted = false;
+						break;
 					}
 				}
-				else
+				// Only use if accepted so far.
+				if (bDeltaAccepted)
 				{
-					ErrorCode = DownloadErrorCodes::UnserialisableDeltaFile;
+					NewManifest = FBuildMergeManifests::MergeDeltaManifest(Configuration.DestinationManifest.Get(), DeltaManifest.Get());
+					if (!NewManifest.IsValid())
+					{
+						UE_LOG(LogOptimisedDelta, Log, TEXT("Rejected optimised delta due to merge failure."));
+						ErrorCode = DownloadErrorCodes::RejectedDeltaFile;
+						bDeltaAccepted = false;
+					}
 				}
+			}
+			else
+			{
+				ErrorCode = DownloadErrorCodes::UnserialisableDeltaFile;
+			}
 
-				if (NewManifest.IsValid())
-				{
-					UE_LOG(LogOptimisedDelta, Log, TEXT("Received optimised delta file successfully %s"), *RelativeDeltaFilePath);
-					DownloadedBytes.Set(Download->GetData().Num());
-					SetResult(NewManifest);
-				}
-				else if (ShouldRetry(Download))
-				{
-					++RetryCount;
-					UE_LOG(LogOptimisedDelta, Log, TEXT("Failed to load %s, retrying %i/%i"), *RelativeDeltaFilePath, RetryCount, DeltaRetries);
-					CloudDirIdx = (CloudDirIdx + RetryCount) % Configuration.CloudDirectories.Num();
-					DownloadService->RequestFile(Configuration.CloudDirectories[CloudDirIdx] / RelativeDeltaFilePath, ChunkDeltaComplete, ChunkDeltaProgress);
-				}
-				else
-				{
-					UE_LOG(LogOptimisedDelta, Log, TEXT("Failed to load %s, ShouldRetry = false"), *RelativeDeltaFilePath);
-					SetFailedDownload();
-				}
+			if (NewManifest.IsValid())
+			{
+				UE_LOG(LogOptimisedDelta, Log, TEXT("Received optimised delta file successfully %s"), *RelativeDeltaFilePath);
+				DownloadedBytes.Set(Download->GetData().Num());
+				SetResult(NewManifest);
 			}
 			else if (ShouldRetry(Download))
 			{
 				++RetryCount;
-				UE_LOG(LogOptimisedDelta, Log, TEXT("Failed to download (response code=%d) %s, retrying %i/%i"), Download->GetResponseCode(), *RelativeDeltaFilePath, RetryCount, DeltaRetries);
-				CloudDirIdx = (CloudDirIdx + RetryCount) % Configuration.CloudDirectories.Num();
-				DownloadService->RequestFile(Configuration.CloudDirectories[CloudDirIdx] / RelativeDeltaFilePath, ChunkDeltaComplete, ChunkDeltaProgress);
+				UE_LOG(LogOptimisedDelta, Log, TEXT("Invalid manifest (response code=%d) %s, retrying %i/%i"), Download->GetResponseCode(), *RelativeDeltaFilePath, RetryCount, DeltaRetries);
+				RequestDeltaFileUri();
 			}
 			else
 			{
-				UE_LOG(LogOptimisedDelta, Log, TEXT("Failed to download (response code=%d) %s, ShouldRetry = false"), Download->GetResponseCode(), *RelativeDeltaFilePath);
-				ErrorCode = DownloadErrorCodes::MissingDeltaFile;
+				UE_LOG(LogOptimisedDelta, Log, TEXT("Invalid manifest (response code=%d) %s, ShouldRetry = false"), Download->GetResponseCode(), *RelativeDeltaFilePath);
 				SetFailedDownload();
 			}
-		}, UE::Tasks::ETaskPriority::BackgroundHigh);
+		}
+		else if (ShouldRetry(Download))
+		{
+			++RetryCount;
+			UE_LOG(LogOptimisedDelta, Log, TEXT("Failed to download (response code=%d) %s, retrying %i/%i"), Download->GetResponseCode(), *RelativeDeltaFilePath, RetryCount, DeltaRetries);
+			RequestDeltaFileUri();
+		}
+		else
+		{
+			UE_LOG(LogOptimisedDelta, Log, TEXT("Failed to download (response code=%d) %s, ShouldRetry = false"), Download->GetResponseCode(), *RelativeDeltaFilePath);
+			ErrorCode = DownloadErrorCodes::MissingDeltaFile;
+			SetFailedDownload();
+		}
+	}
+
+	void FOptimisedDelta::RequestDeltaFileUri()
+	{
+		FUriProviderCompleteDelegate DeltaUriCompletionDelegate = FUriProviderCompleteDelegate::CreateRaw(this, &FOptimisedDelta::DeltaUriCompletionDelegate);
+		Dependencies.UriProvider->RequestUri(RelativeDeltaFilePath, Configuration.DestinationManifest->GetBuildId(), DeltaUriCompletionDelegate);
+	}
+
+	void FOptimisedDelta::DeltaUriCompletionDelegate(bool bSuccessful, const FString& DeltaUri, bool bNeedRetry)
+	{
+		if (bSuccessful)
+		{
+			UE_LOG(LogOptimisedDelta, Log, TEXT("Delta file URI response received: %s"), *DeltaUri);
+			Dependencies.DownloadService->RequestFile(DeltaUri, ChunkDeltaComplete, ChunkDeltaProgress);
+		}
+		else if (bNeedRetry && DeltaUriRetryCount < DeltaUriRetries)
+		{
+			++DeltaUriRetryCount;
+			UE_LOG(LogOptimisedDelta, Log, TEXT("Retrying while trying to get delta URI: %d"), DeltaUriRetryCount);
+			RequestDeltaFileUri();
+		}
+		else
+		{
+			UE_LOG(LogOptimisedDelta, Log, TEXT("No delta file available"));
+			ErrorCode = DownloadErrorCodes::MissingDeltaFile;
+			SetFailedDownload();
+		}
 	}
 
 	bool FOptimisedDelta::ShouldRetry(const FDownloadRef& Download)
@@ -237,19 +313,21 @@ namespace BuildPatchServices
 		// 'no delta' response. Any other failure reason is a server or network issue which we should retry.
 		const int32 ResponseCode = Download->GetResponseCode();
 		const bool bCanRetry = ResponseCode < 400 || ResponseCode >= 500;
+		// In case of URI signature expiration (permission issue) we should retry as well
+		const bool bPermissionError = (ResponseCode == 403);
+		// 404 may occur while working with network shared folders (e.g. when they are not mounted properly).
+		const bool bNotFound = (ResponseCode == 404);
 		const bool bMayRetry = RetryCount < DeltaRetries;
-		return bCanRetry && bMayRetry;
+		return (bCanRetry || bPermissionError || bNotFound) && bMayRetry;
 	}
 
 	void FOptimisedDelta::SetResult(const FBuildPatchAppManifestPtr& ResultManifest)
 	{
-		TSharedPtr<FResultValueOrError> Result = MakeShared<FResultValueOrError>(MakeValue(ResultManifest));
-		TUniqueFunction<void(const IOptimisedDelta::FResultValueOrError&)> LocalOnComplete = MoveTemp(OnComplete);
-
+		TSharedPtr<FResultValueOrError> Result = MakeShared<FResultValueOrError, ESPMode::ThreadSafe>(MakeValue(ResultManifest));
 		ChunkDeltaPromise.SetValue(Result);
-		if (LocalOnComplete)
+		if (Dependencies.OnComplete)
 		{
-			LocalOnComplete(*Result);
+			Dependencies.OnComplete(*Result);
 		}
 	}
 
@@ -264,13 +342,11 @@ namespace BuildPatchServices
 		{
 			UE_LOG(LogOptimisedDelta, Log, TEXT("Failed optimised delta file fetch %s"), *RelativeDeltaFilePath);
 
-			TSharedPtr<FResultValueOrError> Result = MakeShared<FResultValueOrError>(MakeError(ErrorCode));
-			TUniqueFunction<void(const IOptimisedDelta::FResultValueOrError&)> LocalOnComplete = MoveTemp(OnComplete);
-
+			TSharedPtr<FResultValueOrError> Result = MakeShared<FResultValueOrError, ESPMode::ThreadSafe>(MakeError(ErrorCode));
 			ChunkDeltaPromise.SetValue(Result);
-			if (LocalOnComplete)
+			if (Dependencies.OnComplete)
 			{
-				LocalOnComplete(*Result);
+				Dependencies.OnComplete(*Result);
 			}
 		}
 	}
@@ -282,9 +358,18 @@ namespace BuildPatchServices
 	{
 	}
 
-	IOptimisedDelta* FOptimisedDeltaFactory::Create(const FOptimisedDeltaConfiguration& Configuration, FOptimisedDeltaDependencies&& Dependencies)
+	FOptimisedDeltaDependencies::FOptimisedDeltaDependencies(TSharedRef<IUriProvider, ESPMode::ThreadSafe> InUriProvider)
+		: UriProvider(MoveTemp(InUriProvider))
+		, DownloadService(nullptr)
+		, OnComplete([](const IOptimisedDelta::FResultValueOrError&) {})
+	{
+	}
+
+	TSharedRef<IOptimisedDelta> FOptimisedDeltaFactory::Create(const FOptimisedDeltaConfiguration& Configuration, FOptimisedDeltaDependencies Dependencies)
 	{
 		check(Dependencies.DownloadService != nullptr);
-		return new FOptimisedDelta(Configuration, MoveTemp(Dependencies));
+		TSharedRef<FOptimisedDelta> Obj = MakeShared<FOptimisedDelta>(Configuration, MoveTemp(Dependencies));
+		Obj->Initialize();
+		return Obj;
 	}
 }
diff --git a/Engine/Source/Runtime/Online/BuildPatchServices/Private/Installer/OptimisedDelta.h b/Engine/Source/Runtime/Online/BuildPatchServices/Private/Installer/OptimisedDelta.h
--- a/Engine/Source/Runtime/Online/BuildPatchServices/Private/Installer/OptimisedDelta.h
+++ b/Engine/Source/Runtime/Online/BuildPatchServices/Private/Installer/OptimisedDelta.h
@@ -4,6 +4,7 @@
 #include "Templates/ValueOrError.h"
 #include "BuildPatchManifest.h"
 #include "BuildPatchSettings.h"
+#include "UriProvider.h"
 
 namespace BuildPatchServices
 {
@@ -31,6 +32,13 @@ namespace BuildPatchServices
 		 * @return the downloaded size in bytes.
 		 */
 		virtual int32 GetMetaDownloadSize() const = 0;
+
+		/**
+		 * Checks if manifest returned by GetResult() is ready to be provided immediately.
+		 * @return true if result can be obtained.
+		 */
+		virtual bool IsReady() const = 0;
+
 	};
 
 	/**
@@ -49,13 +57,15 @@ namespace BuildPatchServices
 		FBuildPatchAppManifestPtr SourceManifest;
 		// The installation provided destination manifest.
 		FBuildPatchAppManifestRef DestinationManifest;
-		// The list of cloud directory roots that will be used to pull patch data from.
-		TArray<FString> CloudDirectories;
 		// The policy to follow for requesting an optimised delta.
 		EDeltaPolicy DeltaPolicy;
 		// The mode for installation.
 		EInstallMode InstallMode;
-		// The trailer to use for selecting a specific optimized delta or empty if using the default.
+		// A minimal number of retries for delta file download.
+		int32_t RetriesNumber = 0;
+
+		// This is an identifier that we append to the optimized delta filename in order to
+		// facilitate multiple optimized deltas for the same build pairing.
 		FString DeltaFilenameTrailer;
 	};
 
@@ -65,15 +75,23 @@ namespace BuildPatchServices
 	struct FOptimisedDeltaDependencies
 	{
 	public:
+		/**
+		 * Constructor setting up default values.
+		 */
+		FOptimisedDeltaDependencies(TSharedRef<IUriProvider, ESPMode::ThreadSafe> InUriProvider);
+	
+	public:
+		// A URI provider instance.
+		TSharedRef<IUriProvider, ESPMode::ThreadSafe> UriProvider;
 		// A download service instance.
 		IDownloadService* DownloadService = nullptr;
 		// Function to call once the destination manifest has been selected. Receives manifest or error.
-		TUniqueFunction<void(const IOptimisedDelta::FResultValueOrError&)> OnComplete;
+		TUniqueFunction<void(const IOptimisedDelta::FResultValueOrError&)> OnComplete = [](const IOptimisedDelta::FResultValueOrError&) {};
 	};
 
 	class FOptimisedDeltaFactory
 	{
 	public:
-		static IOptimisedDelta* Create(const FOptimisedDeltaConfiguration& Configuration, FOptimisedDeltaDependencies&& Dependencies);
+		static TSharedRef<IOptimisedDelta> Create(const FOptimisedDeltaConfiguration& Configuration, FOptimisedDeltaDependencies Dependencies);
 	};
 }
diff --git a/Engine/Source/Runtime/Online/BuildPatchServices/Private/Installer/Prerequisites.cpp b/Engine/Source/Runtime/Online/BuildPatchServices/Private/Installer/Prerequisites.cpp
--- a/Engine/Source/Runtime/Online/BuildPatchServices/Private/Installer/Prerequisites.cpp
+++ b/Engine/Source/Runtime/Online/BuildPatchServices/Private/Installer/Prerequisites.cpp
@@ -93,8 +93,8 @@ namespace BuildPatchServices
 			TSet<FString> InstalledPrereqs = MachineConfig->LoadInstalledPrereqIds();
 
 			// Check to see if we stored a successful run of this prerequisite already, and can therefore skip it.
-			// We only skip if we are not attempting a repair.
-			if (!PreReq.bIsRepair && PreReq.IdSet.Num() > 0)
+			// We skip if the prerequisites installer already run and we are not attempting a repair.
+			if (Configuration.bSkipPrereqIfAlreadyRan && !PreReq.bIsRepair && PreReq.IdSet.Num() > 0)
 			{
 				const TSet<FString> MissingPrereqs = PreReq.IdSet.Difference(InstalledPrereqs);
 				if (MissingPrereqs.Num() == 0)
@@ -108,10 +108,7 @@ namespace BuildPatchServices
 			bool bUsingStageRoot = false;
 			if (Configuration.InstallMode == EInstallMode::StageFiles)
 			{
-				if (PrereqPath.ReplaceInline(RootDirectoryVariable, *StageDirWithSlash) == 0)
-				{
-					PrereqPath = StageDirWithSlash + PreReq.Path;
-				}
+				PrereqPath = StageDirWithSlash + PreReq.Path;
 				int64 FileSize;
 				if (FileSystem->GetFileSize(*PrereqPath, FileSize))
 				{
@@ -120,10 +117,7 @@ namespace BuildPatchServices
 			}
 			if (!bUsingStageRoot)
 			{
-				if (PrereqPath.ReplaceInline(RootDirectoryVariable, *InstallDirWithSlash) == 0)
-				{
-					PrereqPath = InstallDirWithSlash + PreReq.Path;
-				}
+				PrereqPath = InstallDirWithSlash + PreReq.Path;
 				int64 FileSize;
 				if (FileSystem->GetFileSize(*PrereqPath, FileSize))
 				{
@@ -161,10 +155,35 @@ namespace BuildPatchServices
 			{
 				if (ReturnCode != 0)
 				{
-					UE_LOG(LogPrerequisites, Error, TEXT("Prerequisites executable failed with code %u"), ReturnCode);
-					InstallerAnalytics->RecordPrereqInstallationError(PreReq.AppName, PreReq.VersionString, PrereqPath, PrereqCommandline, ReturnCode, TEXT("Failed to install"));
-					bPrereqInstallSuccessful = false;
-					InstallerError->SetError(EBuildPatchInstallError::PrerequisiteError, *FString::Printf(TEXT("%s%u"), PrerequisiteErrorPrefixes::ReturnCode, ReturnCode));
+					bool bSuspectedSuccess = false;
+					if (!PreReq.bIsRepair)
+					{
+						// Here we try to catch any non-zero error codes that may represent a successful install.  For example, Microsoft installers like VCRedist
+						//  installers that are passed the /norestart commandline arg will not show restart UX but will still return the 3010 error code.  Given this
+						//  info, we can assume a successful run if we see /norestart AND a returned 3010 error code.  Another example would perhaps be if we saw a the
+						//  well known Microsoft 1638 error which for things like VCRedist installer imply a newer version of VCRedist is already installed.  Since the
+						//  newer VCRedist should be backwards compatible with the older version, we can assume this prereq is installed already.
+						if (ReturnCode == 1638)
+						{
+							bSuspectedSuccess = true;
+						}
+						else if (ReturnCode == 3010 && PrereqCommandline.Contains(TEXT("/norestart")))
+						{
+							bSuspectedSuccess = true;
+						}
+					}
+
+					if (bSuspectedSuccess)
+					{
+						UE_LOG(LogPrerequisites, Log, TEXT("Prerequisites executable returned with code %u. This was interpreted as a successful run."), ReturnCode);
+					}
+					else
+					{
+						UE_LOG(LogPrerequisites, Error, TEXT("Prerequisites executable failed with code %u"), ReturnCode);
+						InstallerAnalytics->RecordPrereqInstallationError(PreReq.AppName, PreReq.VersionString, PrereqPath, PrereqCommandline, ReturnCode, TEXT("Failed to install"));
+						bPrereqInstallSuccessful = false;
+						InstallerError->SetError(EBuildPatchInstallError::PrerequisiteError, *FString::Printf(TEXT("%s%u"), PrerequisiteErrorPrefixes::ReturnCode, ReturnCode));
+					}
 				}
 			}
 
diff --git a/Engine/Source/Runtime/Online/BuildPatchServices/Private/Installer/Statistics/CloudChunkSourceStatistics.cpp b/Engine/Source/Runtime/Online/BuildPatchServices/Private/Installer/Statistics/CloudChunkSourceStatistics.cpp
--- a/Engine/Source/Runtime/Online/BuildPatchServices/Private/Installer/Statistics/CloudChunkSourceStatistics.cpp
+++ b/Engine/Source/Runtime/Online/BuildPatchServices/Private/Installer/Statistics/CloudChunkSourceStatistics.cpp
@@ -2,6 +2,7 @@
 
 #include "Installer/Statistics/CloudChunkSourceStatistics.h"
 #include "Misc/ScopeLock.h"
+#include "Algo/Find.h"
 #include "Core/AsyncHelpers.h"
 #include "Common/StatsCollector.h"
 #include "Installer/InstallerAnalytics.h"
@@ -29,11 +30,16 @@ namespace BuildPatchServices
 		virtual void OnDownloadHealthUpdated(EBuildPatchDownloadHealth DownloadHealth) override;
 		virtual void OnSuccessRateUpdated(float SuccessRate) override;
 		virtual void OnActiveRequestCountUpdated(uint32 RequestCount) override;
+		virtual void OnActiveRequestCountAvgUpdated(int32 AverageRequestCount) override;
 		virtual void OnAcceptedNewRequirements(const TSet<FGuid>& ChunkIds) override;
+		virtual void OnChunkDecrypted(uint64 DecryptionCycles) override;
+		virtual void OnCloudDirectoryDataReceived(const FString& CloudDirectory, uint64 BytesDownloaded) override;
+		virtual void OnCloudDirectoryError(const FString& CloudDirectory) override;
 		// ICloudChunkSourceStat interface end.
 
 		// ICloudChunkSourceStatistics interface begin.
 		virtual uint64 GetRequiredDownloadSize() const override;
+		virtual uint64 GetDownloadedBytes() const override;
 		virtual uint64 GetNumCorruptChunkDownloads() const override;
 		virtual uint64 GetNumAbortedChunkDownloads() const override;
 		virtual float GetDownloadSuccessRate() const override;
@@ -41,6 +47,10 @@ namespace BuildPatchServices
 		virtual TArray<float> GetDownloadHealthTimers() const override;
 		virtual uint32 GetCurrentRequestCount() const override;
 		virtual uint32 GetPeakRequestCount() const override;
+		virtual uint32 GetAvgRequestCount() const override;
+		virtual const TCHAR* GetExperimentVariant() const override;
+		virtual double GetTotalDecryptTime() const override;
+		virtual const TArray<FCloudDirectoryStats>& GetCloudDirectoryStats() const override;
 		// ICloudChunkSourceStatistics interface end.
 
 	private:
@@ -54,10 +64,13 @@ namespace BuildPatchServices
 		FThreadSafeInt32 ChunkSuccessRate;
 		FThreadSafeInt32 CurrentRequestCount;
 		volatile uint32 PeakRequestCount;
+		FThreadSafeInt32 AvgRequestCount;
 		mutable FCriticalSection ThreadLockCs;
 		EBuildPatchDownloadHealth CurrentHealth;
 		int64 CyclesAtLastHealthState;
 		TArray<float> HealthStateTimes;
+		FThreadSafeInt64 TotalDecryptCycles;
+		TArray<FCloudDirectoryStats> CloudDirectoryStats;
 	};
 
 	FCloudChunkSourceStatistics::FCloudChunkSourceStatistics(IInstallerAnalytics* InInstallerAnalytics, FBuildPatchProgress* InBuildProgress, IFileOperationTracker* InFileOperationTracker)
@@ -71,9 +84,12 @@ namespace BuildPatchServices
 		, ChunkSuccessRate(0)
 		, CurrentRequestCount(0)
 		, PeakRequestCount(0)
+		, AvgRequestCount(0)
 		, ThreadLockCs()
 		, CurrentHealth(EBuildPatchDownloadHealth::Excellent)
 		, CyclesAtLastHealthState(0)
+		, TotalDecryptCycles(0)
+		, CloudDirectoryStats()
 	{
 		// Initialize health states to zero time.
 		HealthStateTimes.AddZeroed((int32)EBuildPatchDownloadHealth::NUM_Values);
@@ -155,16 +171,63 @@ namespace BuildPatchServices
 		AsyncHelpers::LockFreePeak<int32>((volatile int32*)&PeakRequestCount, (int32)RequestCount);
 	}
 
+	void FCloudChunkSourceStatistics::OnActiveRequestCountAvgUpdated(int32 InAverageRequestCount)
+	{
+		AvgRequestCount.Set(InAverageRequestCount);
+	}
+
 	void FCloudChunkSourceStatistics::OnAcceptedNewRequirements(const TSet<FGuid>& ChunkIds)
 	{
 		FileOperationTracker->OnDataStateUpdate(ChunkIds, EFileOperationState::PendingRemoteCloudData);
 	}
 
+	void FCloudChunkSourceStatistics::OnChunkDecrypted(uint64 DecryptionCycles)
+	{
+		TotalDecryptCycles.Add(DecryptionCycles);
+	}
+
+	void FCloudChunkSourceStatistics::OnCloudDirectoryDataReceived(const FString& CloudDirectory, uint64 ReceivedData)
+	{
+		FCloudDirectoryStats* Stats = Algo::FindBy(CloudDirectoryStats, CloudDirectory, &FCloudDirectoryStats::CloudDirectory);
+		if (Stats != nullptr)
+		{
+			Stats->TotalReceivedData += ReceivedData;
+		}
+		else
+		{
+			FCloudDirectoryStats NewStats;
+			NewStats.CloudDirectory = CloudDirectory;
+			NewStats.TotalReceivedData = ReceivedData;
+			CloudDirectoryStats.Add(NewStats);
+		}
+	}
+
+	void FCloudChunkSourceStatistics::OnCloudDirectoryError(const FString& CloudDirectory)
+	{
+		FCloudDirectoryStats* Stats = Algo::FindBy(CloudDirectoryStats, CloudDirectory, &FCloudDirectoryStats::CloudDirectory);
+		if (Stats != nullptr)
+		{
+			Stats->ErrorCount += 1;
+		}
+		else
+		{
+			FCloudDirectoryStats NewStats;
+			NewStats.CloudDirectory = CloudDirectory;
+			NewStats.ErrorCount = 1;
+			CloudDirectoryStats.Add(NewStats);
+		}
+	}
+
 	uint64 FCloudChunkSourceStatistics::GetRequiredDownloadSize() const
 	{
 		return TotalBytesRequired.GetValue();
 	}
 
+	uint64 FCloudChunkSourceStatistics::GetDownloadedBytes() const
+	{
+		return TotalBytesReceived.GetValue();
+	}
+
 	uint64 FCloudChunkSourceStatistics::GetNumCorruptChunkDownloads() const
 	{
 		return NumDownloadsCorrupt.GetValue();
@@ -202,6 +265,26 @@ namespace BuildPatchServices
 		return PeakRequestCount;
 	}
 
+	uint32 FCloudChunkSourceStatistics::GetAvgRequestCount() const
+	{
+		return AvgRequestCount.GetValue();
+	}
+
+	const TCHAR* FCloudChunkSourceStatistics::GetExperimentVariant() const
+	{
+		return nullptr;
+	}
+
+	double FCloudChunkSourceStatistics::GetTotalDecryptTime() const
+	{
+		return FStatsCollector::CyclesToSeconds(TotalDecryptCycles.GetValue());
+	}
+
+	const TArray<FCloudDirectoryStats>& FCloudChunkSourceStatistics::GetCloudDirectoryStats() const
+	{
+		return CloudDirectoryStats;
+	}
+
 	ICloudChunkSourceStatistics* FCloudChunkSourceStatisticsFactory::Create(IInstallerAnalytics* InstallerAnalytics, FBuildPatchProgress* BuildProgress, IFileOperationTracker* FileOperationTracker)
 	{
 		check(InstallerAnalytics != nullptr);
diff --git a/Engine/Source/Runtime/Online/BuildPatchServices/Private/Installer/Statistics/CloudChunkSourceStatistics.h b/Engine/Source/Runtime/Online/BuildPatchServices/Private/Installer/Statistics/CloudChunkSourceStatistics.h
--- a/Engine/Source/Runtime/Online/BuildPatchServices/Private/Installer/Statistics/CloudChunkSourceStatistics.h
+++ b/Engine/Source/Runtime/Online/BuildPatchServices/Private/Installer/Statistics/CloudChunkSourceStatistics.h
@@ -21,6 +21,11 @@ namespace BuildPatchServices
 		 */
 		virtual uint64 GetRequiredDownloadSize() const = 0;
 
+		/**
+		 * @return the number of bytes received so far.
+		 */
+		virtual uint64 GetDownloadedBytes() const = 0;
+
 		/**
 		 * @return the number of successful chunk downloads which had invalid data.
 		 */
@@ -55,6 +60,26 @@ namespace BuildPatchServices
 		 * @return the peak number of download requests.
 		 */
 		virtual uint32 GetPeakRequestCount() const = 0;
+
+		/**
+		 * @return the avg number of download requests.
+		 */
+		virtual uint32 GetAvgRequestCount() const = 0;
+
+		/**
+		 *  @return the A/B Experiment if applicable
+		 */
+		virtual const TCHAR* GetExperimentVariant() const = 0;
+
+		/**
+		 * @return the total time in seconds spent on chunks decryption.
+		 */
+		virtual double GetTotalDecryptTime() const = 0;
+
+		/**
+		 * @return the stats collected for each cloud directory.
+		 */
+		virtual const TArray<FCloudDirectoryStats>& GetCloudDirectoryStats() const = 0;
 	};
 
 	/**
diff --git a/Engine/Source/Runtime/Online/BuildPatchServices/Private/Installer/Statistics/DownloadServiceStatistics.cpp b/Engine/Source/Runtime/Online/BuildPatchServices/Private/Installer/Statistics/DownloadServiceStatistics.cpp
--- a/Engine/Source/Runtime/Online/BuildPatchServices/Private/Installer/Statistics/DownloadServiceStatistics.cpp
+++ b/Engine/Source/Runtime/Online/BuildPatchServices/Private/Installer/Statistics/DownloadServiceStatistics.cpp
@@ -102,7 +102,7 @@ namespace BuildPatchServices
 			AccumulatedRequestSpeed += Speed;
 			AverageSpeedSampleCount += 1;
 		}
-		else
+		else if (DownloadRecord.FailureReason != EHttpFailureReason::Cancelled)
 		{
 			NumFailedDownloads++;
 			InstallerAnalytics->RecordChunkDownloadError(DownloadRecord.Uri, DownloadRecord.ResponseCode, TEXT("DownloadFail"));
diff --git a/Engine/Source/Runtime/Online/BuildPatchServices/Private/Installer/Statistics/FileOperationTracker.cpp b/Engine/Source/Runtime/Online/BuildPatchServices/Private/Installer/Statistics/FileOperationTracker.cpp
--- a/Engine/Source/Runtime/Online/BuildPatchServices/Private/Installer/Statistics/FileOperationTracker.cpp
+++ b/Engine/Source/Runtime/Online/BuildPatchServices/Private/Installer/Statistics/FileOperationTracker.cpp
@@ -69,6 +69,8 @@ namespace BuildPatchServices
 		void ProcessMessage(FDataState& Message);
 		void ProcessMessage(FFileState& Message);
 		void ProcessMessage(FFileByteRangeState& Message);
+		static void LogOperationState(const FFileOperation* FileOperation);
+		static void LogOperationState(const FFileOperation* FileOperation, const FByteRange& ByteRange);
 
 	private:
 		static FOperationInitialiser BuildOperationInitialiser(IBuildManifestSet* ManifestSet);
@@ -193,6 +195,16 @@ namespace BuildPatchServices
 		return true;
 	}
 
+	void FFileOperationTracker::LogOperationState(const FFileOperation* FileOperation)
+	{
+		UE_LOG(LogFileOperationTracker, Log, TEXT("LogFileOperationTracker: file %s state %s."), *FileOperation->Filename, LexToString(FileOperation->CurrentState));
+	}
+
+	void FFileOperationTracker::LogOperationState(const FFileOperation* FileOperation, const FByteRange& ByteRange)
+	{
+		UE_LOG(LogFileOperationTracker, Log, TEXT("LogFileOperationTracker: file %s state %s start %llu end %llu offset + size %llu."), *FileOperation->Filename, LexToString(FileOperation->CurrentState), ByteRange.Key, ByteRange.Value, FileOperation->Offest + FileOperation->Size);
+	}
+
 	void FFileOperationTracker::ProcessMessage(FOperationInitialiser& Message)
 	{
 		FileOperationStates = MoveTemp(Message.Get<0>());
@@ -219,6 +231,7 @@ namespace BuildPatchServices
 				if (FileOp->CurrentState <= EFileOperationState::DataInMemoryStore)
 				{
 					FileOp->CurrentState = Message.Get<1>();
+					LogOperationState(FileOp);
 				}
 			}
 		}
@@ -231,6 +244,7 @@ namespace BuildPatchServices
 			for (FFileOperation* FileOp : FileOperationStatesFilenameLookup[Message.Get<0>()])
 			{
 				FileOp->CurrentState = Message.Get<1>();
+				LogOperationState(FileOp);
 			}
 		}
 	}
@@ -248,6 +262,7 @@ namespace BuildPatchServices
 				if (FileByteFirst < FileByteEnd && FileByteLast > FileByteStart)
 				{
 					FileOp->CurrentState = Message.Get<2>();
+					LogOperationState(FileOp, Message.Get<1>());
 				}
 				if (FileByteFirst >= FileByteEnd)
 				{
@@ -265,20 +280,26 @@ namespace BuildPatchServices
 		// Get the list of files in the build.
 		TSet<FString> Filenames;
 		ManifestSet->GetExpectedFiles(Filenames);
-		Filenames.Sort(TLess<FString>());
 
 		// Initialise all file operations to Unknown, use dummy operations for empty files.
 		for (const FString& Filename : Filenames)
 		{
-			const FFileManifest* FileManifest = ManifestSet->GetNewFileManifest(Filename);
-			uint64 FileOffset = 0;
-			for (const FChunkPart& FileChunkPart : FileManifest->ChunkParts)
+			if (const FFileManifest* FileManifest = ManifestSet->GetNewFileManifest(Filename))
 			{
-				FileOperationStates.Emplace(Filename, FileChunkPart.Guid, FileOffset, FileChunkPart.Size, EFileOperationState::Unknown);
-				FileOffset += FileChunkPart.Size;
+				uint64 FileOffset = 0;
+				for (const FChunkPart& FileChunkPart : FileManifest->ChunkParts)
+				{
+					FileOperationStates.Emplace(Filename, FileChunkPart.Guid, FileOffset, FileChunkPart.Size, EFileOperationState::Unknown);
+					FileOffset += FileChunkPart.Size;
+				}
+				if (FileManifest->ChunkParts.Num() == 0)
+				{
+					DummyOperationStates.Emplace(Filename, FGuid(), 0, 0, EFileOperationState::Unknown);
+				}
 			}
-			if (FileManifest->ChunkParts.Num() == 0)
+			else
 			{
+				UE_LOG(LogFileOperationTracker, Error, TEXT("LogFileOperationTracker: Failed to get manifest for file %s"), *Filename);
 				DummyOperationStates.Emplace(Filename, FGuid(), 0, 0, EFileOperationState::Unknown);
 			}
 		}
diff --git a/Engine/Source/Runtime/Online/BuildPatchServices/Private/Installer/Statistics/VerifierStatistics.cpp b/Engine/Source/Runtime/Online/BuildPatchServices/Private/Installer/Statistics/VerifierStatistics.cpp
--- a/Engine/Source/Runtime/Online/BuildPatchServices/Private/Installer/Statistics/VerifierStatistics.cpp
+++ b/Engine/Source/Runtime/Online/BuildPatchServices/Private/Installer/Statistics/VerifierStatistics.cpp
@@ -9,6 +9,8 @@
 #include "Installer/Statistics/FileOperationTracker.h"
 #include "BuildPatchVerify.h"
 
+#include <atomic>
+
 namespace VerifierStatisticsHelpers
 {
 	bool TryConvert(BuildPatchServices::EVerifyResult InResult, BuildPatchServices::EVerifyError& OutError)
@@ -56,15 +58,15 @@ namespace BuildPatchServices
 		ISpeedRecorder* SpeedRecorder;
 		FBuildPatchProgress* BuildProgress;
 		IFileOperationTracker* FileOperationTracker;
-		FThreadSafeInt64 TotalBytesProcessed;
-		FThreadSafeInt64 TotalBytesRequired;
-		FThreadSafeInt64 TotalBytesRead;
-		FThreadSafeInt64 NumSuccessfulFilesVerified;
-		FThreadSafeInt64 NumFailedFilesVerified;
+		std::atomic_int64_t TotalBytesProcessed;
+		std::atomic_int64_t TotalBytesRequired;
+		std::atomic_int64_t TotalBytesRead;
+		std::atomic_int64_t NumSuccessfulFilesVerified;
+		std::atomic_int64_t NumFailedFilesVerified;
 		TMap<EVerifyError, FThreadSafeInt64> NumFileVerifyErrors;
-		FThreadSafeBool bIsReading;
-		FThreadSafeInt64 LastReadCycles;
-		uint64 LastByteRangeStateUpdate;
+		std::atomic_bool bIsReading;
+		std::atomic_int64_t LastReadCycles;
+		std::atomic_uint64_t LastByteRangeStateUpdate;
 	};
 
 	FVerifierStatistics::FVerifierStatistics(ISpeedRecorder* InSpeedRecorder, FBuildPatchProgress* InBuildProgress, IFileOperationTracker* InFileOperationTracker)
@@ -94,6 +96,7 @@ namespace BuildPatchServices
 	{
 		bIsReading = true;
 		LastByteRangeStateUpdate = 0;
+		FileOperationTracker->OnFileByteRangeStateUpdate(Filename, FByteRange(LastByteRangeStateUpdate, LastByteRangeStateUpdate), EFileOperationState::Verifying);
 	}
 
 	void FVerifierStatistics::OnFileProgress(const FString& Filename, int64 TotalBytes)
@@ -107,14 +110,14 @@ namespace BuildPatchServices
 		const bool bSuccess = VerifyResult == EVerifyResult::Success;
 		if (bSuccess)
 		{
-			NumSuccessfulFilesVerified.Increment();
+			NumSuccessfulFilesVerified++;
 		}
 		else
 		{
 			EVerifyError Error;
 			if (VerifierStatisticsHelpers::TryConvert(VerifyResult, Error))
 			{
-				NumFailedFilesVerified.Increment();
+				NumFailedFilesVerified++;
 				NumFileVerifyErrors[Error].Increment();
 			}
 		}
@@ -126,13 +129,13 @@ namespace BuildPatchServices
 	void FVerifierStatistics::OnFileRead(const ISpeedRecorder::FRecord& Record)
 	{
 		SpeedRecorder->AddRecord(Record);
-		TotalBytesRead.Add(Record.Size);
+		TotalBytesRead.fetch_add(Record.Size);
 	}
 
 	void FVerifierStatistics::OnProcessedDataUpdated(int64 TotalBytes)
 	{
-		TotalBytesProcessed.Set(TotalBytes);
-		int64 Required = TotalBytesRequired.GetValue();
+		TotalBytesProcessed.store(TotalBytes);
+		int64 Required = TotalBytesRequired.load();
 		if (Required > 0)
 		{
 			BuildProgress->SetStateProgress(EBuildPatchState::BuildVerification, (double)TotalBytes / (double)Required);
@@ -141,8 +144,8 @@ namespace BuildPatchServices
 
 	void FVerifierStatistics::OnTotalRequiredUpdated(int64 TotalBytes)
 	{
-		TotalBytesRequired.Set(TotalBytes);
-		int64 Processed = TotalBytesProcessed.GetValue();
+		TotalBytesRequired.store(TotalBytes);
+		int64 Processed = TotalBytesProcessed.load();
 		if (TotalBytes > 0)
 		{
 			BuildProgress->SetStateProgress(EBuildPatchState::BuildVerification, (double)Processed / (double)TotalBytes);
@@ -151,17 +154,17 @@ namespace BuildPatchServices
 
 	uint64 FVerifierStatistics::GetBytesVerified() const
 	{
-		return TotalBytesRead.GetValue();
+		return TotalBytesRead.load();
 	}
 
 	int32 FVerifierStatistics::GetNumSuccessfulFilesVerified() const
 	{
-		return NumSuccessfulFilesVerified.GetValue();
+		return NumSuccessfulFilesVerified.load();
 	}
 
 	int32 FVerifierStatistics::GetNumFailedFilesVerified() const
 	{
-		return NumFailedFilesVerified.GetValue();
+		return NumFailedFilesVerified.load();
 	}
 
 	TMap<EVerifyError, int32> FVerifierStatistics::GetVerifyErrorCounts() const
@@ -179,7 +182,7 @@ namespace BuildPatchServices
 
 	bool FVerifierStatistics::IsCurrentlyReading() const
 	{
-		const int64 CyclesSinceRead = FStatsCollector::GetCycles() - LastReadCycles.GetValue();
+		const int64 CyclesSinceRead = FStatsCollector::GetCycles() - LastReadCycles.load();
 		return bIsReading || FStatsCollector::CyclesToSeconds(CyclesSinceRead) < 0.05;
 	}
 
diff --git a/Engine/Source/Runtime/Online/BuildPatchServices/Private/Installer/Verifier.cpp b/Engine/Source/Runtime/Online/BuildPatchServices/Private/Installer/Verifier.cpp
--- a/Engine/Source/Runtime/Online/BuildPatchServices/Private/Installer/Verifier.cpp
+++ b/Engine/Source/Runtime/Online/BuildPatchServices/Private/Installer/Verifier.cpp
@@ -4,13 +4,17 @@
 #include "HAL/ThreadSafeBool.h"
 #include "Common/StatsCollector.h"
 #include "Common/FileSystem.h"
+#include "Core/AsyncHelpers.h"
 #include "BuildPatchVerify.h"
 #include "BuildPatchUtil.h"
 #include "HAL/UESemaphore.h"
 #include "IBuildManifestSet.h"
 #include "Misc/ConfigCacheIni.h"
 
-DEFINE_LOG_CATEGORY_STATIC(LogVerifier, Log, All);
+#include <thread>
+
+DECLARE_LOG_CATEGORY_EXTERN(LogVerifier, Warning, All);
+DEFINE_LOG_CATEGORY(LogVerifier);
 
 namespace BuildPatchServices
 {
@@ -41,15 +45,16 @@ namespace BuildPatchServices
 	class FVerifier : public IVerifier
 	{
 	public:
-		FVerifier(IFileSystem* FileSystem, IVerifierStat* InVerificationStat, EVerifyMode InVerifyMode, 
+		FVerifier(IFileSystem* FileSystem, IVerifierStat* InVerificationStat, EVerifyMode InVerifyMode,
 			IBuildInstallerSharedContextPtr InSharedContext, IBuildManifestSet* InManifestSet, 
-			FString InVerifyDirectory, FString InStagedFileDirectory, const TMap<FString, FString>& InPerFileSubdirectories);
-
+			FString InVerifyDirectory, FString InStagedFileDirectory, bool bUseSHA1StageFilenames, 
+			const TMap<FString, FString>& InPerFileSubdirectories);
 		~FVerifier() {}
 
 		// IControllable interface begin.
 		virtual void SetPaused(bool bInIsPaused) override;
 		virtual void Abort() override;
+		virtual void Reset() override;
 		// IControllable interface end.
 
 		// IVerifier interface begin.
@@ -59,8 +64,9 @@ namespace BuildPatchServices
 
 	private:
 		FString SelectFullFilePath(const FString& BuildFile);
-		EVerifyResult VerfiyFileSha(TArray<uint8>& ReadBuffer, const FString& BuildFile, const FFileManifest& BuildFileManifest);
-		EVerifyResult VerfiyFileSize(const FString& BuildFile, const FFileManifest& BuildFileManifest);
+		EVerifyResult VerifyFileSha(TArray<uint8>& ReadBuffer, const FString& BuildFile, const FFileManifest& BuildFileManifest);
+		EVerifyResult VerifyFileSize(const FString& BuildFile, const FFileManifest& BuildFileManifest);
+		EVerifyResult VerifySymlink(const FString& BuildFile, const FFileManifest& BuildFileManifest);
 
 	private:
 		const FString VerifyDirectory;
@@ -68,6 +74,7 @@ namespace BuildPatchServices
 		IFileSystem* const FileSystem;
 		IVerifierStat* const VerifierStat;
 		IBuildManifestSet* const ManifestSet;
+		const bool bUseSHA1StageFilenames;
 		IBuildInstallerSharedContextPtr SharedContext;
 		const TMap<FString, FString>& PerFileSubdirectories;
 
@@ -98,13 +105,13 @@ namespace BuildPatchServices
 
 	FVerifier::FVerifier(IFileSystem* InFileSystem, IVerifierStat* InVerificationStat, EVerifyMode InVerifyMode, 
 		IBuildInstallerSharedContextPtr InSharedContext, IBuildManifestSet* InManifestSet, FString InVerifyDirectory, FString InStagedFileDirectory,
-		const TMap<FString, FString>& InPerFileSubdirectories
-		)
+		bool bUseSHA1StageFilenames, const TMap<FString, FString>& InPerFileSubdirectories)
 		: VerifyDirectory(MoveTemp(InVerifyDirectory))
 		, StagedFileDirectory(MoveTemp(InStagedFileDirectory))
 		, FileSystem(InFileSystem)
 		, VerifierStat(InVerificationStat)
 		, ManifestSet(InManifestSet)
+		, bUseSHA1StageFilenames(bUseSHA1StageFilenames)
 		, SharedContext(InSharedContext)
 		, PerFileSubdirectories(InPerFileSubdirectories)
 		, VerifyMode(InVerifyMode)
@@ -112,7 +119,6 @@ namespace BuildPatchServices
 		, bShouldAbort(false)
 	{
 		ManifestSet->GetFilesTaggedForRepair(FilesToVerify);
-
 	}
 
 	void FVerifier::ProcessVerifyJobs(FSemaphore* ThreadDoneSem)
@@ -147,10 +153,31 @@ namespace BuildPatchServices
 			FileReadBuffer.SetNumUninitialized(4 << 20); // 4MB chunked reads.
 
 			FThreadVerifyJob& OurJob = ThreadJobList[OurJobIndex];
+			if (OurJob.BuildFileManifest == nullptr)
+			{
+				UE_LOG(LogVerifier, Error, TEXT("FVerifier: Unable to find a file manifest for: %s"), *OurJob.FileName);
+				ThreadJobResults[OurJobIndex] = EVerifyResult::FileMissing;
+				continue;
+			}
+
+			const bool bIsSymlink = !OurJob.BuildFileManifest->SymlinkTarget.IsEmpty();
+			const bool bVerifySha = OurJob.bVerifySha;
 
 			VerifierStat->OnFileStarted(OurJob.FileName, OurJob.BuildFileManifest->FileSize);
 
-			EVerifyResult FileVerifyResult = OurJob.bVerifySha ? VerfiyFileSha(FileReadBuffer, OurJob.FileName, *OurJob.BuildFileManifest) : VerfiyFileSize(OurJob.FileName, *OurJob.BuildFileManifest);
+			EVerifyResult FileVerifyResult;
+			if (bIsSymlink)
+			{
+				FileVerifyResult = VerifySymlink(OurJob.FileName, *OurJob.BuildFileManifest);
+			}
+			else if (bVerifySha)
+			{
+				FileVerifyResult = VerifyFileSha(FileReadBuffer, OurJob.FileName, *OurJob.BuildFileManifest);
+			}
+			else
+			{
+				FileVerifyResult = VerifyFileSize(OurJob.FileName, *OurJob.BuildFileManifest);
+			}
 
 			VerifierStat->OnFileCompleted(OurJob.FileName, FileVerifyResult);
 
@@ -176,8 +203,21 @@ namespace BuildPatchServices
 		bShouldAbort = true;
 	}
 
+	void FVerifier::Reset()
+	{
+		bShouldAbort = false;
+	}
+
+	int32 ConfigMaxFileVerifyingTasks(int32 DefaultValue)
+	{
+		int32 ConfigMaxFileVerifyingTasks = DefaultValue;
+		GConfig->GetInt(TEXT("Portal.BuildPatch"), TEXT("MaxFileVerifyingTasks"), ConfigMaxFileVerifyingTasks, GEngineIni);
+		return FMath::Clamp(ConfigMaxFileVerifyingTasks, 1, DefaultValue);
+	}
+
 	EVerifyResult FVerifier::Verify(TArray<FString>& CorruptFiles)
 	{
+		bShouldAbort = false;
 		TRACE_CPUPROFILER_EVENT_SCOPE(Verify);
 		bShouldAbort = false;
 		EVerifyResult VerifyResult = EVerifyResult::Success;
@@ -196,7 +236,8 @@ namespace BuildPatchServices
 
 		// Select verify function.
 		const bool bVerifyShaMode = VerifyMode == EVerifyMode::ShaVerifyAllFiles || VerifyMode == EVerifyMode::ShaVerifyTouchedFiles;
-
+		
+		UE_LOG(LogVerifier, Verbose, TEXT("FVerifier: Setup progress tracking: files %d, mode %s."), VerifyList.Num(), LexToString(VerifyMode));
 
 		ThreadJobList.Reset();
 		ThreadJobResults.Reset();
@@ -295,7 +336,20 @@ namespace BuildPatchServices
 		FString FullFilePath;
 		if (StagedFileDirectory.IsEmpty() == false)
 		{
-			FullFilePath = StagedFileDirectory / BuildFile;
+			FString BuildFilenameToUse = BuildFile;
+			if (bUseSHA1StageFilenames)
+			{
+				const FFileManifest* FileManifest = ManifestSet->GetNewFileManifest(BuildFile);
+				if (FileManifest != nullptr)
+				{
+					FBuildPatchUtils::SHAToBase32(FileManifest->SHA1Hash, BuildFilenameToUse);
+				}
+				else
+				{
+					UE_LOG(LogVerifier, Error, TEXT("FVerifier: Failed to get manifest for file %s"), *BuildFile);
+				}
+			}
+			FullFilePath = StagedFileDirectory / BuildFilenameToUse;
 			int64 FileSize;
 			if (FileSystem->GetFileSize(*FullFilePath, FileSize))
 			{
@@ -306,17 +360,18 @@ namespace BuildPatchServices
 		return FBuildPatchUtils::ResolveInstallationFileName(VerifyDirectory, PerFileSubdirectories, BuildFile);
 	}
 
-	EVerifyResult FVerifier::VerfiyFileSha(TArray<uint8>& ReadBuffer, const FString& BuildFile, const FFileManifest& BuildFileManifest)
+	EVerifyResult FVerifier::VerifyFileSha(TArray<uint8>& ReadBuffer, const FString& BuildFile, const FFileManifest& BuildFileManifest)
 	{
 		TRACE_CPUPROFILER_EVENT_SCOPE(VerifyFileSha);
 		ISpeedRecorder::FRecord ActivityRecord;
 		EVerifyResult VerifyResult;
 		const FString FileToVerify = SelectFullFilePath(BuildFile);
-		uint8 ReturnValue = 0;
 		TUniquePtr<FArchive> FileReader = FileSystem->CreateFileReader(*FileToVerify);
+		int64 FileReaderPrevTell = 0;
 		VerifierStat->OnFileProgress(BuildFile, 0);
 		if (FileReader.IsValid())
 		{
+			FileReaderPrevTell = FileReader->Tell();
 			const int64 FileSize = FileReader->TotalSize();
 			if (FileSize != BuildFileManifest.FileSize)
 			{
@@ -356,12 +411,13 @@ namespace BuildPatchServices
 				}
 				HashState.Final();
 				HashState.GetHash(HashValue.Hash);
-				if (HashValue == BuildFileManifest.FileHash)
+				if (HashValue == BuildFileManifest.SHA1Hash)
 				{
 					VerifyResult = EVerifyResult::Success;
 				}
 				else if (!bShouldAbort)
 				{
+					UE_LOG(LogVerifier, Error, TEXT("FVerifier: Hash check failed for %s. Calculated SHA1: %s, Manifest SHA1: %s"), *BuildFile, *HashValue.ToString(), *BuildFileManifest.SHA1Hash.ToString());
 					VerifyResult = EVerifyResult::HashCheckFailed;
 				}
 				else
@@ -379,6 +435,7 @@ namespace BuildPatchServices
 		{
 			VerifyResult = EVerifyResult::FileMissing;
 		}
+
 		if (VerifyResult != EVerifyResult::Success)
 		{
 			VerifierStat->OnFileProgress(BuildFile, BuildFileManifest.FileSize);
@@ -387,7 +444,7 @@ namespace BuildPatchServices
 		return VerifyResult;
 	}
 
-	EVerifyResult FVerifier::VerfiyFileSize(const FString& BuildFile, const FFileManifest& BuildFileManifest)
+	EVerifyResult FVerifier::VerifyFileSize(const FString& BuildFile, const FFileManifest& BuildFileManifest)
 	{
 		// Pause if necessary.
 		while (bIsPaused && !bShouldAbort)
@@ -396,8 +453,8 @@ namespace BuildPatchServices
 		}
 
 		VerifierStat->OnFileProgress(BuildFile, 0);
-		int64 FileSize;
-		EVerifyResult VerifyResult;
+		int64 FileSize = 0;
+		EVerifyResult VerifyResult = EVerifyResult::FileSizeFailed;
 		if (FileSystem->GetFileSize(*SelectFullFilePath(BuildFile), FileSize))
 		{
 			if (FileSize == BuildFileManifest.FileSize)
@@ -419,13 +476,54 @@ namespace BuildPatchServices
 		return VerifyResult;
 	}
 
+	EVerifyResult FVerifier::VerifySymlink(const FString& BuildFile, const FFileManifest& BuildFileManifest)
+	{
+		VerifierStat->OnFileProgress(BuildFile, 0);
+		EVerifyResult VerifyResult = EVerifyResult::HashCheckFailed;
+#if PLATFORM_MAC || PLATFORM_LINUX
+		// Load the symlink
+		TArray<ANSICHAR> SymlinkTargetAnsi;
+		SymlinkTargetAnsi.AddZeroed(FPlatformMisc::GetMaxPathLength() + 1);
+		if (readlink(TCHAR_TO_UTF8(*SelectFullFilePath(BuildFile)), SymlinkTargetAnsi.GetData(), SymlinkTargetAnsi.Num() - 1) != -1)
+		{
+			// Check the path is correct
+			const FString SymlinkTarget = UTF8_TO_TCHAR(SymlinkTargetAnsi.GetData());
+			if (SymlinkTarget == BuildFileManifest.SymlinkTarget)
+			{
+				VerifyResult = EVerifyResult::Success;
+			}
+			else
+			{
+				VerifyResult = EVerifyResult::HashCheckFailed;
+			}
+		}
+		// If the symlink does not exist
+		else if (errno == ENOENT)
+		{
+			VerifyResult = EVerifyResult::FileMissing;
+		}
+		// Otherwise we treat as not being able to open the link (e.g. permission error or not a link)
+		else
+		{
+			VerifyResult = EVerifyResult::OpenFileFailed;
+		}
+#else
+		// Cannot verify on unsupported platforms
+		VerifyResult = EVerifyResult::Success;
+#endif
+		VerifierStat->OnFileProgress(BuildFile, BuildFileManifest.FileSize);
+
+		ThreadProcessedBytes.fetch_add(BuildFileManifest.FileSize, std::memory_order_release);
+		return VerifyResult;
+	}
+
 	IVerifier* FVerifierFactory::Create(IFileSystem* FileSystem, IVerifierStat* VerifierStat, EVerifyMode VerifyMode, 
 		IBuildInstallerSharedContextPtr SharedContext, IBuildManifestSet* ManifestSet, FString VerifyDirectory, FString StagedFileDirectory,
-		const TMap<FString, FString>& PerFileSubdirectories
-		)
+		bool bUseSHA1StageFilenames, const TMap<FString, FString>& PerFileSubdirectories)
 	{
 		check(FileSystem != nullptr);
 		check(VerifierStat != nullptr);
-		return new FVerifier(FileSystem, VerifierStat, VerifyMode, SharedContext, ManifestSet, MoveTemp(VerifyDirectory), MoveTemp(StagedFileDirectory), PerFileSubdirectories);
+		check(ManifestSet != nullptr);
+		return new FVerifier(FileSystem, VerifierStat, VerifyMode, SharedContext, ManifestSet, MoveTemp(VerifyDirectory), MoveTemp(StagedFileDirectory), bUseSHA1StageFilenames, PerFileSubdirectories);
 	}
 }
diff --git a/Engine/Source/Runtime/Online/BuildPatchServices/Private/Installer/Verifier.h b/Engine/Source/Runtime/Online/BuildPatchServices/Private/Installer/Verifier.h
--- a/Engine/Source/Runtime/Online/BuildPatchServices/Private/Installer/Verifier.h
+++ b/Engine/Source/Runtime/Online/BuildPatchServices/Private/Installer/Verifier.h
@@ -79,16 +79,16 @@ namespace BuildPatchServices
 	{
 	public:
 		/**
-		* Creates a verifier class that will verify a local directory structure against a given manifest, optionally
-		* taking account of a staging directory where alternative files are used.
-		* @param StagedFileDirectory   A stage directory for updated files, ignored if empty string. If a file exists here, it will be checked instead of the one in VerifyDirectory.
-		* @param PerFileSubdirectories An map of path adjustments to use keyed off the manifest filename, leave empty if unneeded.
-		* @return The verifier object. Delete when complete. See IVerifier class for usage.
-		*/
+		 * Creates a verifier class that will verify a local directory structure against a given manifest, optionally
+		 * taking account of a staging directory where alternative files are used.
+		 * @param StagedFileDirectory    A stage directory for updated files, ignored if empty string. If a file exists here, it will be checked instead of the one in VerifyDirectory.
+		 * @param bUseSHA1StageFilenames A boolean value to indicate that we use hashes instead of real file names.
+		 * @param PerFileSubdirectories  A map of path adjustments to use keyed off the manifest filename, leave empty if unneeded.
+		 * @return The verifier object. Delete when complete. See IVerifier class for usage.
+		 */
 		static IVerifier* Create(IFileSystem* FileSystem, IVerifierStat* VerifierStat, EVerifyMode VerifyMode, 
 			IBuildInstallerSharedContextPtr SharedContext, IBuildManifestSet* ManifestSet, FString VerifyDirectory, FString StagedFileDirectory,
-			const TMap<FString, FString>& PerFileSubdirectories
-			);
+			bool bUseSHA1StageFilenames, const TMap<FString, FString>& PerFileSubdirectories);
 	};
 
 	/**
diff --git a/Engine/Source/Runtime/Online/BuildPatchServices/Private/Tests/Fake/FileSystem.fake.h b/Engine/Source/Runtime/Online/BuildPatchServices/Private/Tests/Fake/FileSystem.fake.h
--- a/Engine/Source/Runtime/Online/BuildPatchServices/Private/Tests/Fake/FileSystem.fake.h
+++ b/Engine/Source/Runtime/Online/BuildPatchServices/Private/Tests/Fake/FileSystem.fake.h
@@ -81,6 +81,12 @@ namespace BuildPatchServices
 			return Reader;
 		}
 
+		virtual TUniquePtr<FArchive> CreateFileReaderEx(const TCHAR* Filename, EReadFlags ReadFlags, IFileHandle** OutHandle) const override
+		{
+			check(OutHandle == nullptr); // The output parameter is not supported for fake implementation yet.
+			return CreateFileReader(Filename, ReadFlags);
+		}
+
 		virtual TUniquePtr<FArchive> CreateFileWriter(const TCHAR* Filename, EWriteFlags WriteFlags = EWriteFlags::None) const override
 		{
 			TUniquePtr<FArchive> Writer;
diff --git a/Engine/Source/Runtime/Online/BuildPatchServices/Private/Tests/Mock/ChunkDataSerialization.mock.h b/Engine/Source/Runtime/Online/BuildPatchServices/Private/Tests/Mock/ChunkDataSerialization.mock.h
--- a/Engine/Source/Runtime/Online/BuildPatchServices/Private/Tests/Mock/ChunkDataSerialization.mock.h
+++ b/Engine/Source/Runtime/Online/BuildPatchServices/Private/Tests/Mock/ChunkDataSerialization.mock.h
@@ -39,7 +39,7 @@ namespace BuildPatchServices
 			return Result;
 		}
 
-		virtual EChunkSaveResult SaveToFile(const FString& Filename, const IChunkDataAccess* ChunkDataAccess) const override
+		virtual EChunkSaveResult SaveToFile(const FString& Filename, IChunkDataAccess* ChunkDataAccess) const override
 		{
 			RxSaveToFile.Emplace(Filename, ChunkDataAccess);
 			return EChunkSaveResult::Success;
@@ -58,7 +58,7 @@ namespace BuildPatchServices
 			return Result;
 		}
 
-		virtual EChunkSaveResult SaveToMemory(TArray<uint8>& Memory, const IChunkDataAccess* ChunkDataAccess) const override
+		virtual EChunkSaveResult SaveToMemory(TArray<uint8>& Memory, IChunkDataAccess* ChunkDataAccess) const override
 		{
 			MOCK_FUNC_NOT_IMPLEMENTED("FMockChunkDataSerialization::SaveToMemory");
 			return EChunkSaveResult::SerializationError;
@@ -76,8 +76,22 @@ namespace BuildPatchServices
 			RxLoadFromArchive.Emplace(OutLoadResult);
 			return Result;
 		}
+		
+		virtual TFunction<TPair<EChunkLoadResult, IChunkDataAccess*>()> LoadFromArchiveAsync(FArchive& Archive) const override
+		{
+			EChunkLoadResult LoadResult;
+			IChunkDataAccess* Result = nullptr;
+			if (TxLoadFromArchive.Num())
+			{
+				FTxLoadFromFile LoadFromArchiveResult = TxLoadFromArchive.Pop();
+				Result = LoadFromArchiveResult.Get<0>();
+				LoadResult = LoadFromArchiveResult.Get<1>();
+			}
+			RxLoadFromArchive.Emplace(LoadResult);
+			return [LoadResult, Result]() {return TPair<EChunkLoadResult, IChunkDataAccess*>(LoadResult, Result); };
+		}
 
-		virtual EChunkSaveResult SaveToArchive(FArchive& Archive, const IChunkDataAccess* ChunkDataAccess) const override
+		virtual EChunkSaveResult SaveToArchive(FArchive& Archive, IChunkDataAccess* ChunkDataAccess) const override
 		{
 			EChunkSaveResult Result = EChunkSaveResult::Success;
 			if (SaveToArchiveFunc)
@@ -88,16 +102,17 @@ namespace BuildPatchServices
 			return Result;
 		}
 
-		virtual void InjectShaToChunkData(TArray<uint8>& Memory, const FSHAHash& ShaHashData) const override
+		virtual bool InjectShaToChunkData(TArray<uint8>& Memory, const FSHAHash& ShaHashData) const override
 		{
 			RxInjectShaToChunkData.Emplace(Memory, ShaHashData);
+			return true;
 		}
 
 		virtual bool ValidateAndRead(FArchive& InArchive, FMutableMemoryView InDestinationBuffer, FChunkHeader& OutHeader, FUniqueBuffer& OutCompressedBuffer) const
 		{
 			return false;
 		}
-		virtual bool DecompressValidatedRead(const FChunkHeader& InHeader, FMutableMemoryView InDestionationBuffer, const FUniqueBuffer& InCompressedBuffer) const
+		virtual bool DecompressAndDecryptValidatedRead(const FChunkHeader& InHeader, FMutableMemoryView InDestionationBuffer, const FUniqueBuffer& InCompressedBuffer) const
 		{
 			return false;
 		}
diff --git a/Engine/Source/Runtime/Online/BuildPatchServices/Private/Tests/Mock/CloudChunkSourceStat.mock.h b/Engine/Source/Runtime/Online/BuildPatchServices/Private/Tests/Mock/CloudChunkSourceStat.mock.h
--- a/Engine/Source/Runtime/Online/BuildPatchServices/Private/Tests/Mock/CloudChunkSourceStat.mock.h
+++ b/Engine/Source/Runtime/Online/BuildPatchServices/Private/Tests/Mock/CloudChunkSourceStat.mock.h
@@ -22,6 +22,9 @@ namespace BuildPatchServices
 		typedef TTuple<double, EBuildPatchDownloadHealth> FDownloadHealthUpdated;
 		typedef TTuple<double, float> FSuccessRateUpdated;
 		typedef TTuple<double, int32> FActiveRequestCountUpdated;
+		typedef TTuple<double, int32> FActiveRequestCountAverageUpdated;
+		typedef TTuple<double, FString, uint64> FCloudDirectoryDataReceived;
+		typedef TTuple<double, FString> FCloudDirectoryError;
 
 	public:
 		virtual void OnDownloadRequested(const FGuid& ChunkId) override
@@ -74,10 +77,29 @@ namespace BuildPatchServices
 			RxActiveRequestCountUpdated.Emplace(FStatsCollector::GetSeconds(), RequestCount);
 		}
 
+		virtual void OnActiveRequestCountAvgUpdated(int32 AverageRequestCount) override
+		{
+			RxActiveRequestCountAverageUpdated.Emplace(FStatsCollector::GetSeconds(), AverageRequestCount);
+		}
+
 		virtual void OnAcceptedNewRequirements(const TSet<FGuid>& ChunkIds) override
 		{
 		}
 
+		virtual void OnChunkDecrypted(uint64 DecryptionCycles) override
+		{
+		}
+
+		virtual void OnCloudDirectoryDataReceived(const FString& CloudDirectory, uint64 ReceivedData) override
+		{
+			RxCloudDirectoryDataReceived.Emplace(FStatsCollector::GetSeconds(), CloudDirectory, ReceivedData);
+		}
+
+		virtual void OnCloudDirectoryError(const FString& CloudDirectory) override
+		{
+			RxCloudDirectoryError.Emplace(FStatsCollector::GetSeconds(), CloudDirectory);
+		}
+
 	public:
 		TArray<FDownloadRequested> RxDownloadRequested;
 		TArray<FDownloadSuccess> RxDownloadSuccess;
@@ -89,6 +111,9 @@ namespace BuildPatchServices
 		TArray<FDownloadHealthUpdated> RxDownloadHealthUpdated;
 		TArray<FSuccessRateUpdated> RxSuccessRateUpdated;
 		TArray<FActiveRequestCountUpdated> RxActiveRequestCountUpdated;
+		TArray<FActiveRequestCountUpdated> RxActiveRequestCountAverageUpdated;
+		TArray<FCloudDirectoryDataReceived> RxCloudDirectoryDataReceived;
+		TArray<FCloudDirectoryError> RxCloudDirectoryError;
 	};
 }
 
diff --git a/Engine/Source/Runtime/Online/BuildPatchServices/Private/Tests/Mock/DownloadConnectionCount.mock.h b/Engine/Source/Runtime/Online/BuildPatchServices/Private/Tests/Mock/DownloadConnectionCount.mock.h
--- a/Engine/Source/Runtime/Online/BuildPatchServices/Private/Tests/Mock/DownloadConnectionCount.mock.h
+++ b/Engine/Source/Runtime/Online/BuildPatchServices/Private/Tests/Mock/DownloadConnectionCount.mock.h
@@ -13,7 +13,7 @@ namespace BuildPatchServices
 	{
 	public:
 		FMockDownloadConnectionCount(uint32 InCount)
-			:Count(InCount)
+			: Count(InCount)
 		{
 
 		}
@@ -23,6 +23,7 @@ namespace BuildPatchServices
 
 		}
 		uint32 GetAdjustedCount(uint32 NumProcessing, EBuildPatchDownloadHealth CurrentHealth) { return Count; }
+		void SetConnectionScaling(bool bInEnableScaling) { }
 	private:
 		const uint32 Count;
 	};
diff --git a/Engine/Source/Runtime/Online/BuildPatchServices/Private/Tests/Mock/DownloadService.mock.h b/Engine/Source/Runtime/Online/BuildPatchServices/Private/Tests/Mock/DownloadService.mock.h
--- a/Engine/Source/Runtime/Online/BuildPatchServices/Private/Tests/Mock/DownloadService.mock.h
+++ b/Engine/Source/Runtime/Online/BuildPatchServices/Private/Tests/Mock/DownloadService.mock.h
@@ -21,25 +21,21 @@ namespace BuildPatchServices
 		{
 		}
 
-		virtual int32 RequestFile(const FString& FileUri, const FDownloadCompleteDelegate& OnCompleteDelegate, const FDownloadProgressDelegate& OnProgressDelegate) override
+		virtual bool RequestFile(const FString& FileUri, const FDownloadCompleteDelegate& OnCompleteDelegate, const FDownloadProgressDelegate& OnProgressDelegate, int32* OutRequestId = nullptr) override
 		{
 			FScopeLock ScopeLock(&ThreadLock);
-			int32 ReturnId;
+			bool bSuccess = false;
 			if (RequestFileFunc)
 			{
-				ReturnId = RequestFileFunc(FileUri, OnCompleteDelegate, OnProgressDelegate);
+				bSuccess = RequestFileFunc(FileUri, OnCompleteDelegate, OnProgressDelegate, OutRequestId);
 			}
-			else
-			{
-				ReturnId = ++Count;
-			}
-			RxRequestFile.Emplace(FStatsCollector::GetSeconds(), ReturnId, FileUri, OnCompleteDelegate, OnProgressDelegate);
-			return ReturnId;
+			RxRequestFile.Emplace(FStatsCollector::GetSeconds(), OutRequestId ? *OutRequestId : -1, FileUri, OnCompleteDelegate, OnProgressDelegate);
+			return bSuccess;
 		}
 
-		virtual int32 RequestFileWithHeaders(const FString& FileUri, const TMap<FString, FString>& Headers, const FDownloadCompleteDelegate& OnCompleteDelegate, const FDownloadProgressDelegate& OnProgressDelegate) override
+		virtual bool RequestFileWithHeaders(const FString& FileUri, const TMap<FString, FString>& Headers, const FDownloadCompleteDelegate& OnCompleteDelegate, const FDownloadProgressDelegate& OnProgressDelegate, int32* OutRequestId = nullptr) override
 		{
-			return RequestFile(FileUri, OnCompleteDelegate, OnProgressDelegate);
+			return RequestFile(FileUri, OnCompleteDelegate, OnProgressDelegate, OutRequestId);
 		}
 
 		virtual void RequestCancel(int32 RequestId) override
@@ -67,7 +63,7 @@ namespace BuildPatchServices
 		int32 Count;
 		TArray<FRequestFile> RxRequestFile;
 		TArray<FRequestCancel> RxRequestCancel;
-		TFunction<int32(const FString&, const FDownloadCompleteDelegate&, const FDownloadProgressDelegate&)> RequestFileFunc;
+		TFunction<bool(const FString&, const FDownloadCompleteDelegate&, const FDownloadProgressDelegate&, int32*)> RequestFileFunc;
 		TFunction<void(int32)> RequestCancelFunc;
 	};
 }
diff --git a/Engine/Source/Runtime/Online/BuildPatchServices/Private/Tests/Mock/FileSystem.mock.h b/Engine/Source/Runtime/Online/BuildPatchServices/Private/Tests/Mock/FileSystem.mock.h
--- a/Engine/Source/Runtime/Online/BuildPatchServices/Private/Tests/Mock/FileSystem.mock.h
+++ b/Engine/Source/Runtime/Online/BuildPatchServices/Private/Tests/Mock/FileSystem.mock.h
@@ -33,6 +33,12 @@ namespace BuildPatchServices
 			return Reader;
 		}
 
+		virtual TUniquePtr<FArchive> CreateFileReaderEx(const TCHAR* Filename, EReadFlags ReadFlags, IFileHandle** OutHandle) const override
+		{
+			check(OutHandle == nullptr); // The output parameter is not supported for fake implementation yet.
+			return CreateFileReader(Filename, ReadFlags);
+		}
+
 		virtual TUniquePtr<FArchive> CreateFileWriter(const TCHAR* Filename, EWriteFlags WriteFlags = EWriteFlags::None) const override
 		{
 			FScopeLock ScopeLock(&ThreadLock);
@@ -53,6 +59,18 @@ namespace BuildPatchServices
 			return false;
 		}
 
+		virtual bool LoadFileToArray(const TCHAR* Filename, TArray<uint8>& Contents) const override
+		{
+			MOCK_FUNC_NOT_IMPLEMENTED("FMockFileSystem::LoadFileToArray");
+			return false;
+		}
+
+		virtual bool SaveArrayToFile(const TCHAR* Filename, const TArray<uint8>& Contents) const override
+		{
+			MOCK_FUNC_NOT_IMPLEMENTED("FMockFileSystem::SaveArrayToFile");
+			return false;
+		}
+
 		virtual bool DirectoryExists(const TCHAR* DirectoryPath) const override
 		{
 			MOCK_FUNC_NOT_IMPLEMENTED("FMockFileSystem::DirectoryExists");
@@ -65,13 +83,19 @@ namespace BuildPatchServices
 			return false;
 		}
 
+		virtual bool DeleteDirectoryRecursively(const TCHAR* DirectoryPath) const override
+		{
+			MOCK_FUNC_NOT_IMPLEMENTED("FMockFileSystem::DeleteDirectoryRecursively");
+			return false;
+		}
+
 		virtual bool DeleteFile(const TCHAR* Filename) const override
 		{
 			MOCK_FUNC_NOT_IMPLEMENTED("FMockFileSystem::DeleteFile");
 			return false;
 		}
 
-		virtual bool MoveFile(const TCHAR* FileDest, const TCHAR* FileSource) const override
+		virtual bool MoveFile(const TCHAR* FileDest, const TCHAR* FileSource, bool bDoNotRetryOrError) const override
 		{
 			MOCK_FUNC_NOT_IMPLEMENTED("FMockFileSystem::MoveFile");
 			return false;
@@ -107,7 +131,7 @@ namespace BuildPatchServices
 
 		virtual bool GetTimeStamp(const TCHAR* Path, FDateTime& TimeStamp) const override
 		{
-			MOCK_FUNC_NOT_IMPLEMENTED("FMockFileSystem::CopyFile");
+			MOCK_FUNC_NOT_IMPLEMENTED("FMockFileSystem::GetTimeStamp");
 			return false;
 		}
 
@@ -129,6 +153,12 @@ namespace BuildPatchServices
 			return true;
 		}
 
+		virtual bool SetFileSize(const TCHAR* Filename, int64 FileSize) const override
+		{
+			MOCK_FUNC_NOT_IMPLEMENTED("FMockFileSystem::SetFileSize");
+			return false;
+		}
+
 		virtual bool FileExists(const TCHAR* Filename) const override
 		{
 			MOCK_FUNC_NOT_IMPLEMENTED("FMockFileSystem::FileExists");
@@ -141,15 +171,36 @@ namespace BuildPatchServices
 			return;
 		}
 
-		virtual void FindFilesRecursively(TArray<FString>& FoundFiles, const TCHAR* Directory, const TCHAR* FileExtension) const override
+		virtual void FindFilesRecursively(TArray<FString>& FoundFiles, const TCHAR* Directory, const TCHAR* FileExtension, bool bFollowSymlinks = true) const override
 		{
 			MOCK_FUNC_NOT_IMPLEMENTED("FMockFileSystem::FindFilesRecursively");
 			return;
 		}
 
-		virtual void ParallelFindFilesRecursively(TArray<FString>& FoundFiles, const TCHAR* Directory, const TCHAR* FileExtension = nullptr, EAsyncExecution AsyncExecution = EAsyncExecution::ThreadPool) const override
+		virtual void ParallelFindFilesRecursively(TArray<FString>& FoundFiles, const TCHAR* Directory, const TCHAR* FileExtension = nullptr, EAsyncExecution AsyncExecution = EAsyncExecution::ThreadPool, bool bFollowSymlinks = true) const override
+		{
+			FindFilesRecursively(FoundFiles, Directory, FileExtension, bFollowSymlinks);
+		}
+
+		virtual FFileStatData GetStatData(const TCHAR* FilenameOrDirectory) const override
 		{
-			FindFilesRecursively(FoundFiles, Directory, FileExtension);
+			if (!FileExists(FilenameOrDirectory))
+			{
+				return {};
+			}
+			FFileStatData Result;
+			Result.bIsDirectory = false;
+			Result.bIsValid = true;
+
+			Result.bIsValid &= GetFileSize(FilenameOrDirectory, Result.FileSize);
+			Result.bIsValid &= GetTimeStamp(FilenameOrDirectory, Result.ModificationTime);
+			Result.AccessTime = Result.ModificationTime;
+			Result.CreationTime = Result.ModificationTime;
+
+			EAttributeFlags Attributes;
+			Result.bIsValid &= GetAttributes(FilenameOrDirectory, Attributes);
+			Result.bIsReadOnly = EnumHasAllFlags(Attributes, EAttributeFlags::ReadOnly);
+			return Result;
 		}
 
 		virtual int64 GetAllowedBytesToWriteThrottledStorage(const TCHAR* DestinationPath = nullptr) const override
diff --git a/Engine/Source/Runtime/Online/BuildPatchServices/Private/Tests/Mock/Manifest.mock.h b/Engine/Source/Runtime/Online/BuildPatchServices/Private/Tests/Mock/Manifest.mock.h
--- a/Engine/Source/Runtime/Online/BuildPatchServices/Private/Tests/Mock/Manifest.mock.h
+++ b/Engine/Source/Runtime/Online/BuildPatchServices/Private/Tests/Mock/Manifest.mock.h
@@ -131,6 +131,13 @@ namespace BuildPatchServices
 			return false;
 		}
 
+		virtual TArray<FString> GetCustomFieldNames() const override
+		{
+			TArray<FString> CustomFieldNames;
+			CustomFields.GenerateKeyArray(CustomFieldNames);
+			return CustomFieldNames;
+		}
+
 		virtual const IManifestFieldPtr GetCustomField(const FString& FieldName) const override
 		{
 			return MakeShareable(new FMockManifestField(CustomFields[FieldName]));
@@ -176,7 +183,7 @@ namespace BuildPatchServices
 			return true;
 		}
 
-		virtual bool SaveToFile(const FString& Filename, EFeatureLevel InFeatureLevel) override
+		virtual bool SaveToFile(const FString& Filename, EFeatureLevel InFeatureLevel, FSHAHash* OutSHA1Hash, FMD5Hash* OutMD5Hash) override
 		{
 			return true;
 		}
@@ -235,6 +242,31 @@ namespace BuildPatchServices
 			return GetFileSize(Filenames.Array());
 		}
 
+		virtual EFileMetaFlags GetFileMetaFlags(const FString& Filename) const override
+		{
+			return FileManifests[Filename].FileMetaFlags;
+		}
+
+		virtual const FMD5Hash* GetFileMD5Hash(const FString& Filename) const override
+		{
+			return &FileManifests[Filename].MD5Hash;
+		}
+
+		virtual const FSHAHash* GetFileSHA1Hash(const FString& Filename) const override
+		{
+			return &FileManifests[Filename].SHA1Hash;
+		}
+
+		virtual const FSHA256Signature* GetFileSHA256Hash(const FString& Filename) const override
+		{
+			return &FileManifests[Filename].SHA256Hash;
+		}
+
+		virtual const FString* GetFileMIMEType(const FString& Filename) const override
+		{
+			return &FileManifests[Filename].MIMEType;
+		}
+
 		virtual uint32 GetNumFiles() const override
 		{
 			return NumFiles;
@@ -311,6 +343,16 @@ namespace BuildPatchServices
 			return false;
 		}
 
+		virtual bool GetChunkEncryptionSecretId(const FGuid& ChunkGuid, FGuid& OutEncryptionSecretId) const override
+		{
+			if (ChunkInfos.Contains(ChunkGuid))
+			{
+				OutEncryptionSecretId = ChunkInfos[ChunkGuid].EncryptionSecretId;
+				return OutEncryptionSecretId.IsValid();
+			}
+			return false;
+		}
+
 		virtual const FChunkInfo* GetChunkInfo(const FGuid& ChunkGuid) const override
 		{
 			if (ChunkInfos.Contains(ChunkGuid))
diff --git a/Engine/Source/Runtime/Online/BuildPatchServices/Private/Tests/Mock/MessagePump.mock.h b/Engine/Source/Runtime/Online/BuildPatchServices/Private/Tests/Mock/MessagePump.mock.h
--- a/Engine/Source/Runtime/Online/BuildPatchServices/Private/Tests/Mock/MessagePump.mock.h
+++ b/Engine/Source/Runtime/Online/BuildPatchServices/Private/Tests/Mock/MessagePump.mock.h
@@ -23,9 +23,15 @@ namespace BuildPatchServices
 		{
 		}
 
+		virtual void PumpMessages(const TArray<FMessageHandler*>& Handlers) override
+		{
+		}
 		virtual void PumpMessages() override
 		{
 		}
+		virtual void SendRequest(FDeltaUriRequest Request, TFunction<void(FExpiratedUrisResponse)> OnResponse) override
+		{
+		}
 
 		virtual bool SendRequest(FChunkUriRequest Request, TFunction<void(FChunkUriResponse)> OnResponse) override
 		{
diff --git a/Engine/Source/Runtime/Online/BuildPatchServices/Private/Tests/TestHelpers.h b/Engine/Source/Runtime/Online/BuildPatchServices/Private/Tests/TestHelpers.h
--- a/Engine/Source/Runtime/Online/BuildPatchServices/Private/Tests/TestHelpers.h
+++ b/Engine/Source/Runtime/Online/BuildPatchServices/Private/Tests/TestHelpers.h
@@ -79,6 +79,14 @@ namespace BuildPatchTestHelpers
 #define TEST_NOT_NULL(expression) \
 	_TEST_NOT_NULL(TEXT(#expression), expression)
 
+#define TEST_SIZE(ar, value) \
+	_TEST_EQUAL(TEXT("size check"), ar.Num(), value)
+
+#define TEST_EMPTY(ar) TEST_SIZE(ar, 0)
+
+#define TEST_CONTAINS(expression, value) \
+	_TEST_EQUAL(TEXT("contains check"), expression.Contains(value), true)
+
 #define TEST_BECOMES_TRUE(expression, timeout) \
 	BuildPatchTestHelpers::WaitUntilTrue([this](){ return expression; }, [this](){ return TEST_EQUAL(expression, true); }, timeout)
 
diff --git a/Engine/Source/Runtime/Online/BuildPatchServices/Private/Tests/Unit/CloudChunkSource.spec.cpp b/Engine/Source/Runtime/Online/BuildPatchServices/Private/Tests/Unit/CloudChunkSource.spec.cpp
--- a/Engine/Source/Runtime/Online/BuildPatchServices/Private/Tests/Unit/CloudChunkSource.spec.cpp
+++ b/Engine/Source/Runtime/Online/BuildPatchServices/Private/Tests/Unit/CloudChunkSource.spec.cpp
@@ -233,13 +233,18 @@ void FCloudChunkSourceSpec::Define()
 				CloudSourceConfig->PreFetchMaximum = InitialDownloadSet.Num() + 1;
 				CloudSourceConfig->bBeginDownloadsOnFirstGet = false;
 				ForceConnectionCount = 1;
-				FakeDownloadService->RequestFileFunc = [this](const FString&, const FDownloadCompleteDelegate&, const FDownloadProgressDelegate&)
+				FakeDownloadService->RequestFileFunc = [this](const FString&, const FDownloadCompleteDelegate&, const FDownloadProgressDelegate&, int32* OutRequestId) -> bool
 				{
 					if (++CallCount == 2)
 					{
 						PauseFor(PauseTime);
 					}
-					return ++FakeDownloadService->Count;
+					const int32 NewRequestId = ++FakeDownloadService->Count;
+					if (OutRequestId)
+					{
+						*OutRequestId = NewRequestId;
+					}
+					return true;
 				};
 				MakeUnit(InitialDownloadSet);
 				FakeDownloadService->StartService();
