p4 integrate //UE5/Dev-Distro-5.5/Engine/Source/Programs/BuildPatchTool/... //UE5/Main/Engine/Source/Programs/BuildPatchTool/...

diff --git a/Engine/Source/Runtime/Online/BuildPatchServices/Private/Tests/Unit/DownloadService.spec.cpp b/Engine/Source/Runtime/Online/BuildPatchServices/Private/Tests/Unit/DownloadService.spec.cpp
--- a/Engine/Source/Runtime/Online/BuildPatchServices/Private/Tests/Unit/DownloadService.spec.cpp
+++ b/Engine/Source/Runtime/Online/BuildPatchServices/Private/Tests/Unit/DownloadService.spec.cpp
@@ -17,10 +17,10 @@ BEGIN_DEFINE_SPEC(FDownloadServiceSpec, "BuildPatchServices.Unit", EAutomationTe
 TUniquePtr<BuildPatchServices::IDownloadService> DownloadService;
 // Mock/Fake.
 FTSTicker Ticker;
-TUniquePtr<BuildPatchServices::FFakeHttpManager> FakeHttpModule;
-TUniquePtr<BuildPatchServices::FMockFileSystem> MockFileSystem;
-TUniquePtr<BuildPatchServices::FMockDownloadServiceStat> MockDownloadServiceStat;
-TUniquePtr<BuildPatchServices::FMockInstallerAnalytics> MockInstallerAnalytics;
+TSharedPtr<BuildPatchServices::FFakeHttpManager> FakeHttpModule;
+TSharedPtr<BuildPatchServices::FMockFileSystem> MockFileSystem;
+TSharedPtr<BuildPatchServices::FMockDownloadServiceStat> MockDownloadServiceStat;
+TSharedPtr<BuildPatchServices::FMockInstallerAnalytics> MockInstallerAnalytics;
 // Data.
 BuildPatchServices::FDownloadProgressDelegate DownloadProgress;
 BuildPatchServices::FDownloadCompleteDelegate DownloadComplete;
@@ -32,8 +32,11 @@ FString NetworkFileUrl;
 int32 MadeRequestId;
 // Ticker helpers.
 void DoTick();
-void DoTicksUntilCreated(int32 TickCount = 50, int32 CreateCount = 1, float Sleep = 0.0f);
-void DoTicksUntilComplete(int32 TickCount = 50, int32 CompleteCount = 1, float Sleep = 0.0f);
+bool DoTicksUntilCreated(float Seconds = 5.0);
+bool DoTicksUntilComplete(float Seconds = 5.0, int32 CompleteCount = 1);
+bool DoTicksUntil(float Seconds, const TFunction<bool()>& pred);
+int RxCreateFileReaderNum() const;
+
 END_DEFINE_SPEC(FDownloadServiceSpec)
 
 void FDownloadServiceSpec::Define()
@@ -58,15 +61,15 @@ void FDownloadServiceSpec::Define()
 	BeforeEach([this]()
 	{
 		Ticker.Reset();
-		FakeHttpModule.Reset(new FFakeHttpManager(Ticker));
-		MockFileSystem.Reset(new FMockFileSystem());
-		MockDownloadServiceStat.Reset(new FMockDownloadServiceStat());
-		MockInstallerAnalytics.Reset(new FMockInstallerAnalytics());
+		FakeHttpModule = MakeShared<FFakeHttpManager>(Ticker);
+		MockFileSystem = MakeShared<FMockFileSystem>();
+		MockDownloadServiceStat = MakeShared<FMockDownloadServiceStat>();
+		MockInstallerAnalytics = MakeShared<FMockInstallerAnalytics>();
 		DownloadService.Reset(FDownloadServiceFactory::Create(
-			FakeHttpModule.Get(),
-			MockFileSystem.Get(),
-			MockDownloadServiceStat.Get(),
-			MockInstallerAnalytics.Get()));
+			FakeHttpModule,
+			MockFileSystem,
+			MockDownloadServiceStat,
+			MockInstallerAnalytics));
 	});
 
 	xDescribe("DownloadService", [this]()
@@ -78,7 +81,7 @@ void FDownloadServiceSpec::Define()
 				It("should use the http module to process the request.", [this]()
 				{
 					DownloadService->RequestFile(HttpFileUrl, DownloadComplete, DownloadProgress);
-					DoTicksUntilComplete();
+					TEST_TRUE(DoTicksUntilComplete());
 					TEST_EQUAL(FakeHttpModule->RxCreateRequest, 1);
 					TEST_EQUAL(MockFileSystem->RxCreateFileReader.Num(), 0);
 				});
@@ -89,7 +92,7 @@ void FDownloadServiceSpec::Define()
 				It("should use the http module to process the request.", [this]()
 				{
 					DownloadService->RequestFile(HttpsFileUrl, DownloadComplete, DownloadProgress);
-					DoTicksUntilComplete();
+					TEST_TRUE(DoTicksUntilComplete());
 					TEST_EQUAL(FakeHttpModule->RxCreateRequest, 1);
 					TEST_EQUAL(MockFileSystem->RxCreateFileReader.Num(), 0);
 				});
@@ -100,7 +103,8 @@ void FDownloadServiceSpec::Define()
 				xIt("should use the file manager to process the request.", [this]()
 				{
 					DownloadService->RequestFile(NetworkFileUrl, DownloadComplete, DownloadProgress);
-					DoTicksUntilComplete(50, 1);
+					TEST_TRUE(DoTicksUntilComplete());
+					TEST_TRUE(DoTicksUntilComplete());
 					TEST_EQUAL(FakeHttpModule->RxCreateRequest, 0);
 					TEST_EQUAL(MockFileSystem->RxCreateFileReader.Num(), 1);
 				});
@@ -116,7 +120,7 @@ void FDownloadServiceSpec::Define()
 				It("should provide an IDownload with access to success status.", [this]()
 				{
 					DownloadService->RequestFile(HttpFileUrl, DownloadComplete, DownloadProgress);
-					DoTicksUntilComplete();
+					TEST_TRUE(DoTicksUntilComplete());
 					TEST_EQUAL(RxDownloadComplete.Num(), 1);
 					if (RxDownloadComplete.Num() == 1)
 					{
@@ -138,7 +142,7 @@ void FDownloadServiceSpec::Define()
 				xIt("should provide an IDownload with access to success status.", [this]()
 				{
 					DownloadService->RequestFile(NetworkFileUrl, DownloadComplete, DownloadProgress);
-					DoTicksUntilComplete();
+					TEST_TRUE(DoTicksUntilComplete());
 					TEST_EQUAL(RxDownloadComplete.Num(), 1);
 					if (RxDownloadComplete.Num() == 1)
 					{
@@ -158,13 +162,13 @@ void FDownloadServiceSpec::Define()
 				BeforeEach([this]()
 				{
 					MockFileSystem->ReadFile.AddUninitialized(1024 * 1024 * 50);
-					MadeRequestId = DownloadService->RequestFile(NetworkFileUrl, DownloadComplete, DownloadProgress);
+					DownloadService->RequestFile(NetworkFileUrl, DownloadComplete, DownloadProgress, &MadeRequestId);
 				});
 
 				It("should cancel if it was not started yet.", [this]()
 				{
 					DownloadService->RequestCancel(MadeRequestId);
-					DoTicksUntilComplete();
+					TEST_TRUE(DoTicksUntilComplete());
 					TEST_EQUAL(FakeHttpModule->RxCreateRequest, 0);
 					TEST_EQUAL(MockFileSystem->RxCreateFileReader.Num(), 0);
 					TEST_EQUAL(RxDownloadComplete.Num(), 1);
@@ -177,11 +181,10 @@ void FDownloadServiceSpec::Define()
 
 				xIt("should cancel if it had already started.", [this]()
 				{
-					DoTicksUntilCreated();
+					TEST_TRUE(DoTicksUntilCreated());
 					DownloadService->RequestCancel(MadeRequestId);
-					DoTicksUntilComplete(50, 1, 0.1f);
+					TEST_TRUE(DoTicksUntilComplete());
 					TEST_EQUAL(FakeHttpModule->RxCreateRequest, 0);
-					TEST_EQUAL(MockFileSystem->RxCreateFileReader.Num(), 1);
 					TEST_EQUAL(RxDownloadComplete.Num(), 1);
 					if (RxDownloadComplete.Num() == 1)
 					{
@@ -196,13 +199,13 @@ void FDownloadServiceSpec::Define()
 				BeforeEach([this]()
 				{
 					FakeHttpModule->DataServed.FindOrAdd(HttpFileUrl).AddUninitialized(1024*1024*50);
-					MadeRequestId = DownloadService->RequestFile(HttpFileUrl, DownloadComplete, DownloadProgress);
+					DownloadService->RequestFile(HttpFileUrl, DownloadComplete, DownloadProgress, &MadeRequestId);
 				});
 
 				It("should cancel if it was not started yet.", [this]()
 				{
 					DownloadService->RequestCancel(MadeRequestId);
-					DoTicksUntilComplete();
+					TEST_TRUE(DoTicksUntilComplete());
 					TEST_EQUAL(FakeHttpModule->RxCreateRequest, 0);
 					TEST_EQUAL(MockFileSystem->RxCreateFileReader.Num(), 0);
 					TEST_EQUAL(RxDownloadComplete.Num(), 1);
@@ -215,9 +218,9 @@ void FDownloadServiceSpec::Define()
 
 				It("should cancel if it had already started.", [this]()
 				{
-					DoTicksUntilCreated();
+					TEST_TRUE(DoTicksUntilCreated());
 					DownloadService->RequestCancel(MadeRequestId);
-					DoTicksUntilComplete();
+					TEST_TRUE(DoTicksUntilComplete());
 					TEST_EQUAL(FakeHttpModule->RxCreateRequest, 1);
 					TEST_EQUAL(MockFileSystem->RxCreateFileReader.Num(), 0);
 					TEST_EQUAL(RxDownloadComplete.Num(), 1);
@@ -236,10 +239,10 @@ void FDownloadServiceSpec::Define()
 			{
 				It("should cancel the request.", [this]()
 				{
-					int32 RequestId = DownloadService->RequestFile(HttpFileUrl, DownloadComplete, DownloadProgress);
-					DoTicksUntilCreated();
+					DownloadService->RequestFile(HttpFileUrl, DownloadComplete, DownloadProgress);
+					TEST_TRUE(DoTicksUntilCreated());
 					DownloadService.Reset();
-					DoTicksUntilComplete();
+					TEST_TRUE(DoTicksUntilComplete());
 					TEST_EQUAL(FakeHttpModule->RxCreateRequest, 1);
 					TEST_EQUAL(MockFileSystem->RxCreateFileReader.Num(), 0);
 					TEST_EQUAL(RxDownloadComplete.Num(), 1);
@@ -270,28 +273,36 @@ void FDownloadServiceSpec::DoTick()
 	Ticker.Tick(0.1f);
 }
 
-void FDownloadServiceSpec::DoTicksUntilComplete(int32 TickCount, int32 CompleteCount, float Sleep)
+bool FDownloadServiceSpec::DoTicksUntilComplete(float Seconds, int32 CompleteCount)
 {
-	while (TickCount > 0 && RxDownloadComplete.Num() < CompleteCount)
-	{
-		DoTick();
-		--TickCount;
-		FPlatformProcess::Sleep(Sleep);
-	}
+	return DoTicksUntil(Seconds, [this, CompleteCount]() {return RxDownloadComplete.Num() == CompleteCount; });
 }
 
-void FDownloadServiceSpec::DoTicksUntilCreated(int32 TickCount, int32 CreateCount, float Sleep)
+bool FDownloadServiceSpec::DoTicksUntil(float Seconds, const TFunction<bool()>& pred)
 {
-	MockFileSystem->ThreadLock.Lock();
-	while (TickCount > 0 && (FakeHttpModule->RxCreateRequest + MockFileSystem->RxCreateFileReader.Num()) < CreateCount)
+	FDateTime Start = FDateTime::UtcNow();
+	do
 	{
-		MockFileSystem->ThreadLock.Unlock();
+		if (pred())
+		{
+			return true;
+		}
 		DoTick();
-		--TickCount;
-		FPlatformProcess::Sleep(Sleep);
-		MockFileSystem->ThreadLock.Lock();
-	}
-	MockFileSystem->ThreadLock.Unlock();
+		FPlatformProcess::Sleep(.005);
+	} while ((FDateTime::UtcNow() - Start) < FTimespan::FromSeconds(Seconds));
+
+	return pred();
+}
+
+bool FDownloadServiceSpec::DoTicksUntilCreated(float Seconds)
+{
+	return DoTicksUntil(Seconds, [this]() { return FakeHttpModule->RxCreateRequest == 1 || RxCreateFileReaderNum() == 1; });
+}
+
+int FDownloadServiceSpec::RxCreateFileReaderNum() const
+{
+	FScopeLock Lock(&MockFileSystem->ThreadLock);
+	return MockFileSystem->RxCreateFileReader.Num();
 }
 
 #endif //WITH_DEV_AUTOMATION_TESTS
diff --git a/Engine/Source/Runtime/Online/BuildPatchServices/Private/Tests/Unit/FileAttribution.spec.cpp b/Engine/Source/Runtime/Online/BuildPatchServices/Private/Tests/Unit/FileAttribution.spec.cpp
--- a/Engine/Source/Runtime/Online/BuildPatchServices/Private/Tests/Unit/FileAttribution.spec.cpp
+++ b/Engine/Source/Runtime/Online/BuildPatchServices/Private/Tests/Unit/FileAttribution.spec.cpp
@@ -335,11 +335,11 @@ void FFileAttributionSpec::MakeUnit()
 		ChunkPart.Size = 64;
 		FFileManifest FileManifest;
 		FileManifest.Filename = File;
-		FSHA1::HashBuffer(*FGuid::NewGuid().ToString(), sizeof(TCHAR) * 32, FileManifest.FileHash.Hash);
+		FSHA1::HashBuffer(*FGuid::NewGuid().ToString(), sizeof(TCHAR) * 32, FileManifest.SHA1Hash.Hash);
 		FileManifest.ChunkParts.Add(ChunkPart);
-		FileManifest.FileMetaFlags |= ReadOnlyFiles.Contains(File) ? BuildPatchServices::EFileMetaFlags::ReadOnly : BuildPatchServices::EFileMetaFlags::None;
-		FileManifest.FileMetaFlags |= CompressedFiles.Contains(File) ? BuildPatchServices::EFileMetaFlags::Compressed : BuildPatchServices::EFileMetaFlags::None;
-		FileManifest.FileMetaFlags |= ExeFiles.Contains(File) ? BuildPatchServices::EFileMetaFlags::UnixExecutable : BuildPatchServices::EFileMetaFlags::None;
+		FileManifest.FileMetaFlags |= ReadOnlyFiles.Contains(File) ? EFileMetaFlags::ReadOnly : EFileMetaFlags::None;
+		FileManifest.FileMetaFlags |= CompressedFiles.Contains(File) ? EFileMetaFlags::Compressed : EFileMetaFlags::None;
+		FileManifest.FileMetaFlags |= ExeFiles.Contains(File) ? EFileMetaFlags::UnixExecutable : EFileMetaFlags::None;
 		FileManifest.FileSize = Algo::Accumulate<int64>(FileManifest.ChunkParts, 0, [] (int64 Count, const FChunkPart& Element) { return Count + Element.Size; });
 		MockNewManifest->BuildFileList.Add(File);
 		MockNewManifest->FileManifests.Add(File, FileManifest);
@@ -350,7 +350,7 @@ void FFileAttributionSpec::MakeUnit()
 		}
 		else if (ChangedFiles.Contains(File))
 		{
-			FSHA1::HashBuffer(*FGuid::NewGuid().ToString(), sizeof(TCHAR) * 32, FileManifest.FileHash.Hash);
+			FSHA1::HashBuffer(*FGuid::NewGuid().ToString(), sizeof(TCHAR) * 32, FileManifest.SHA1Hash.Hash);
 			MockOldManifest->BuildFileList.Add(File);
 			MockOldManifest->FileManifests.Add(File, FileManifest);
 		}
diff --git a/Engine/Source/Runtime/Online/BuildPatchServices/Private/Tests/Unit/InstallChunkSource.spec.cpp b/Engine/Source/Runtime/Online/BuildPatchServices/Private/Tests/Unit/InstallChunkSource.spec.cpp
--- a/Engine/Source/Runtime/Online/BuildPatchServices/Private/Tests/Unit/InstallChunkSource.spec.cpp
+++ b/Engine/Source/Runtime/Online/BuildPatchServices/Private/Tests/Unit/InstallChunkSource.spec.cpp
@@ -335,6 +335,7 @@ void FInstallChunkSourceSpec::MakeUnit()
 		MockInstallChunkSourceStat.Get(),
 		InstallationSources,
 		ChunksThatWillBeNeeded,
+		false,
 		EmptyPerFileSubdirectories));
 }
 
@@ -405,14 +406,14 @@ void FInstallChunkSourceSpec::InventUsableChunkData()
 					*((uint32*)(Data + DataIdx)) = RandomData.GetUnsignedInt();
 				}
 
-				FSHA1::HashBuffer(Data, FileManifest.FileSize, FileManifest.FileHash.Hash);
+				FSHA1::HashBuffer(Data, FileManifest.FileSize, FileManifest.SHA1Hash.Hash);
 
 				// Also fill the chunk array so we can hash it later.
 				FMemory::Memcpy(&ChunkData[ChunkPart.Offset], Data, ChunkPart.Size);
 
 				ThisMockManifest->BuildFileList.Add(FileManifest.Filename);
 				ThisMockManifest->FileNameToFileSize.Add(FileManifest.Filename, FileManifest.FileSize);
-				ThisMockManifest->FileNameToHashes.Add(FileManifest.Filename,FileManifest.FileHash);
+				ThisMockManifest->FileNameToHashes.Add(FileManifest.Filename,FileManifest.SHA1Hash);
 				ThisMockManifest->FileManifests.Add(FileManifest.Filename, MoveTemp(FileManifest));
 			}
 
@@ -426,7 +427,7 @@ void FInstallChunkSourceSpec::InventUsableChunkData()
 			ChunkInfo.Hash = ChunkPolyHash;
 			ChunkInfo.ShaHash = ChunkShaHash;
 			ChunkInfo.GroupNumber = 0;
-			ChunkInfo.WindowSize = TestChunkSize;
+			ChunkInfo.DataSizeUncompressed = TestChunkSize;
 			ChunkInfo.FileSize = TestChunkSize;
 			ThisMockManifest->ChunkInfos.Add(TheChunk, ChunkInfo);
 		}
diff --git a/Engine/Source/Runtime/Online/BuildPatchServices/Public/BuildPatchDelta.h b/Engine/Source/Runtime/Online/BuildPatchServices/Public/BuildPatchDelta.h
--- a/Engine/Source/Runtime/Online/BuildPatchServices/Public/BuildPatchDelta.h
+++ b/Engine/Source/Runtime/Online/BuildPatchServices/Public/BuildPatchDelta.h
@@ -29,7 +29,7 @@ namespace BuildPatchServices
 	 * Based on the source and destination manifests, get the filename for the delta that optimises patching from source to destination.
 	 * @param SourceManifest        The source manifest.
 	 * @param DestinationManifest   The destination manifest.
-	 * @param FilenameTrailer		The trailer to use for selecting a non standard optimized delta, or empty if unneeded.
+	 * @param FilenameTrailer       This is an optional parameter to allow for multiple optimized deltas for the same build pairing. If unneeded, leave as "".
 	 * @return the CloudDir relative delta filename.
 	 */
 	BUILDPATCHSERVICES_API FString GetChunkDeltaFilename(const IBuildManifestRef& SourceManifest, const IBuildManifestRef& DestinationManifest, const FString& FilenameTrailer);
@@ -42,13 +42,13 @@ static_assert((uint32)BuildPatchServices::EDeltaPolicy::InvalidOrMax == 3, "Plea
 inline const TCHAR* LexToString(BuildPatchServices::EDeltaPolicy DeltaPolicy)
 {
 #define CASE_ENUM_TO_STR(Value) case BuildPatchServices::EDeltaPolicy::Value: return TEXT(#Value)
-		switch (DeltaPolicy)
-		{
+	switch (DeltaPolicy)
+	{
 		CASE_ENUM_TO_STR(TryFetchContinueWithout);
 		CASE_ENUM_TO_STR(Expect);
 		CASE_ENUM_TO_STR(Skip);
-	default: return TEXT("InvalidOrMax");
-		}
+		default: return TEXT("InvalidOrMax");
+	}
 #undef CASE_ENUM_TO_STR
 }
 
diff --git a/Engine/Source/Runtime/Online/BuildPatchServices/Public/BuildPatchFeatureLevel.h b/Engine/Source/Runtime/Online/BuildPatchServices/Public/BuildPatchFeatureLevel.h
--- a/Engine/Source/Runtime/Online/BuildPatchServices/Public/BuildPatchFeatureLevel.h
+++ b/Engine/Source/Runtime/Online/BuildPatchServices/Public/BuildPatchFeatureLevel.h
@@ -2,7 +2,6 @@
 #pragma once
 
 #include "CoreMinimal.h"
-#include "HAL/Platform.h"
 
 namespace BuildPatchServices
 {
@@ -48,6 +47,20 @@ namespace BuildPatchServices
 		UsesRuntimeGeneratedBuildId,
 		// Manifest uses a build id generated unique at build time, and stored in manifest.
 		UsesBuildTimeGeneratedBuildId,
+		// Manifests generated with this feature level onwards will store the MD5 hash and the calculated MIME type of each file.
+		StoresFileMD5HashesAndMIMEType,
+		// Manifests generated with this feature level onwards will store the SHA256 hash of each file.
+		StoresFileSHA256Hashes,
+		// Added support for Uninstall Actions.
+		StoresUninstallActions,
+		// Added full support for chunk encryption, additionally stores encyption secrets, as well as the CompressedDataSize and AuthTag
+		// for each chunk. Chunks from here onwards are stored in ChunksV5.
+		ChunkEncryptionSupport,
+		// Added the secretID as part of the chunk pathing within a CloudDir
+		ChunksStoredBySecret,
+		// Completed full support for BPS data encryption, if a manifest got encrypted, it will additionally store encryption data
+		// and most fields will have been replaced by empty or functional but not accurate data.
+		ManifestEncryptionSupport,
 
 		// !! Always after the latest version entry, signifies the latest version plus 1 to allow the following Latest alias.
 		LatestPlusOne,
@@ -57,6 +70,10 @@ namespace BuildPatchServices
 		LatestNoChunks = StoresChunkFileSizes,
 		// An alias to provide the latest version of a manifest supported by a json serialized format.
 		LatestJson = StoresPrerequisiteIds,
+		// An alias to provide the latest version of a manifest supported by generation runs on platforms without OpenSSL module support.
+		LatestNoOpenSSL = StoresFileMD5HashesAndMIMEType,
+		// An alias to provide the latest version of a manifest with unencrypted chunks.
+		LatestUnencryptedChunks = StoresUninstallActions,
 		// An alias to provide the first available version of optimised delta manifest saving.
 		FirstOptimisedDelta = UsesRuntimeGeneratedBuildId,
 
@@ -69,19 +86,89 @@ namespace BuildPatchServices
 		// This is for UObject default, so that we always serialize it.
 		Invalid = -1
 	};
+}
 
-	/**
-	 * Returns the string representation of the EFeatureLevel value. Used for analytics and logging only.
-	 * @param FeatureLevel  The feature level enum value.
-	 * @return The string representation.
-	 */
-	BUILDPATCHSERVICES_API const TCHAR* FeatureLevelToString(const EFeatureLevel& FeatureLevel);
+inline const TCHAR* LexToString(BuildPatchServices::EFeatureLevel FeatureLevel)
+{
+	static_assert((int32)BuildPatchServices::EFeatureLevel::Latest == 24, "Please add support for the extra values below.");
+#define CASE_ENUM_TO_STR(Value) case BuildPatchServices::EFeatureLevel::Value: return TEXT(#Value)
+	switch (FeatureLevel)
+	{
+		CASE_ENUM_TO_STR(Original);
+		CASE_ENUM_TO_STR(CustomFields);
+		CASE_ENUM_TO_STR(StartStoringVersion);
+		CASE_ENUM_TO_STR(DataFileRenames);
+		CASE_ENUM_TO_STR(StoresIfChunkOrFileData);
+		CASE_ENUM_TO_STR(StoresDataGroupNumbers);
+		CASE_ENUM_TO_STR(ChunkCompressionSupport);
+		CASE_ENUM_TO_STR(StoresPrerequisitesInfo);
+		CASE_ENUM_TO_STR(StoresChunkFileSizes);
+		CASE_ENUM_TO_STR(StoredAsCompressedUClass);
+		CASE_ENUM_TO_STR(UNUSED_0);
+		CASE_ENUM_TO_STR(UNUSED_1);
+		CASE_ENUM_TO_STR(StoresChunkDataShaHashes);
+		CASE_ENUM_TO_STR(StoresPrerequisiteIds);
+		CASE_ENUM_TO_STR(StoredAsBinaryData);
+		CASE_ENUM_TO_STR(VariableSizeChunksWithoutWindowSizeChunkInfo);
+		CASE_ENUM_TO_STR(VariableSizeChunks);
+		CASE_ENUM_TO_STR(UsesRuntimeGeneratedBuildId);
+		CASE_ENUM_TO_STR(UsesBuildTimeGeneratedBuildId);
+		CASE_ENUM_TO_STR(StoresFileMD5HashesAndMIMEType);
+		CASE_ENUM_TO_STR(StoresFileSHA256Hashes);
+		CASE_ENUM_TO_STR(StoresUninstallActions);
+		CASE_ENUM_TO_STR(ChunkEncryptionSupport);
+		CASE_ENUM_TO_STR(ChunksStoredBySecret);
+		CASE_ENUM_TO_STR(ManifestEncryptionSupport);
+		CASE_ENUM_TO_STR(BrokenJsonVersion);
+		default: return TEXT("InvalidOrMax");
+	}
+#undef CASE_ENUM_TO_STR
+}
 
-	/**
-	 * Parses the provided string into the relevant EFeatureLevel value, if it matches.
-	 * @param FeatureLevelString    The string to try parse.
-	 * @param FeatureLevel          Receives the enum value if successful.
-	 * @return true if successfully parsed.
-	 */
-	BUILDPATCHSERVICES_API bool FeatureLevelFromString(const TCHAR* FeatureLevelString, EFeatureLevel& FeatureLevel);
+inline void LexFromString(BuildPatchServices::EFeatureLevel& FeatureLevel, const TCHAR* Buffer)
+{
+	static_assert((int32)BuildPatchServices::EFeatureLevel::Latest == 24, "Please add support for the extra values below.");
+#define RETURN_IF_EQUAL(Value) if (FCString::Stricmp(Buffer, TEXT(#Value)) == 0) { FeatureLevel = BuildPatchServices::EFeatureLevel::Value; return; }
+	const TCHAR* const Prefix = TEXT("EFeatureLevel::");
+	const SIZE_T PrefixLen = FCString::Strlen(Prefix);
+	if (FCString::Strnicmp(Buffer, Prefix, PrefixLen) == 0)
+	{
+		Buffer += PrefixLen;
+	}
+	RETURN_IF_EQUAL(Original);
+	RETURN_IF_EQUAL(CustomFields);
+	RETURN_IF_EQUAL(StartStoringVersion);
+	RETURN_IF_EQUAL(DataFileRenames);
+	RETURN_IF_EQUAL(StoresIfChunkOrFileData);
+	RETURN_IF_EQUAL(StoresDataGroupNumbers);
+	RETURN_IF_EQUAL(ChunkCompressionSupport);
+	RETURN_IF_EQUAL(StoresPrerequisitesInfo);
+	RETURN_IF_EQUAL(StoresChunkFileSizes);
+	RETURN_IF_EQUAL(StoredAsCompressedUClass);
+	RETURN_IF_EQUAL(UNUSED_0);
+	RETURN_IF_EQUAL(UNUSED_1);
+	RETURN_IF_EQUAL(StoresChunkDataShaHashes);
+	RETURN_IF_EQUAL(StoresPrerequisiteIds);
+	RETURN_IF_EQUAL(StoredAsBinaryData);
+	RETURN_IF_EQUAL(VariableSizeChunksWithoutWindowSizeChunkInfo);
+	RETURN_IF_EQUAL(VariableSizeChunks);
+	RETURN_IF_EQUAL(UsesRuntimeGeneratedBuildId);
+	RETURN_IF_EQUAL(UsesBuildTimeGeneratedBuildId);
+	RETURN_IF_EQUAL(StoresFileMD5HashesAndMIMEType);
+	RETURN_IF_EQUAL(StoresFileSHA256Hashes);
+	RETURN_IF_EQUAL(StoresUninstallActions);
+	RETURN_IF_EQUAL(ChunkEncryptionSupport);
+	RETURN_IF_EQUAL(ChunksStoredBySecret);
+	RETURN_IF_EQUAL(ManifestEncryptionSupport);
+	RETURN_IF_EQUAL(LatestPlusOne);
+	RETURN_IF_EQUAL(Latest);
+	RETURN_IF_EQUAL(LatestNoChunks);
+	RETURN_IF_EQUAL(LatestJson);
+	RETURN_IF_EQUAL(LatestUnencryptedChunks);
+	RETURN_IF_EQUAL(FirstOptimisedDelta);
+	RETURN_IF_EQUAL(StoresUniqueBuildId);
+	RETURN_IF_EQUAL(BrokenJsonVersion);
+	// Did not match
+	FeatureLevel = BuildPatchServices::EFeatureLevel::Invalid;
+#undef RETURN_IF_EQUAL
 }
diff --git a/Engine/Source/Runtime/Online/BuildPatchServices/Public/BuildPatchInstall.h b/Engine/Source/Runtime/Online/BuildPatchServices/Public/BuildPatchInstall.h
--- a/Engine/Source/Runtime/Online/BuildPatchServices/Public/BuildPatchInstall.h
+++ b/Engine/Source/Runtime/Online/BuildPatchServices/Public/BuildPatchInstall.h
@@ -26,13 +26,16 @@ namespace BuildPatchServices
 		// Execute the prerequisite installer only, downloading it first if necessary. If the specified manifest has no prerequisites, this will result in an error.
 		PrereqOnly,
 
+		// Download chunk data to file only, to be unpacked as a later process. Chunk data will not be decrypted or decompressed.
+		Preload,
+
 		InvalidOrMax
 	};
 
 	BUILDPATCHSERVICES_API uint64 CalculateRequiredDiskSpace(const IBuildManifestPtr& CurrentManifest, const IBuildManifestRef& BuildManifest, const EInstallMode& InstallMode, const TSet<FString>& InstallTags);
 }
 
-static_assert((uint32)BuildPatchServices::EInstallMode::InvalidOrMax == 4, "Please add support for the extra values to the Lex functions below.");
+static_assert((uint32)BuildPatchServices::EInstallMode::InvalidOrMax == 5, "Please add support for the extra values to the Lex functions below.");
 
 inline const TCHAR* LexToString(BuildPatchServices::EInstallMode InstallMode)
 {
@@ -43,8 +46,9 @@ inline const TCHAR* LexToString(BuildPatchServices::EInstallMode InstallMode)
 		CASE_ENUM_TO_STR(DestructiveInstall);
 		CASE_ENUM_TO_STR(NonDestructiveInstall);
 		CASE_ENUM_TO_STR(PrereqOnly);
-	default: return TEXT("InvalidOrMax");
-		}
+		CASE_ENUM_TO_STR(Preload);
+		default: return TEXT("InvalidOrMax");
+	}
 #undef CASE_ENUM_TO_STR
 }
 
@@ -61,6 +65,7 @@ inline void LexFromString(BuildPatchServices::EInstallMode& InstallMode, const T
 	RETURN_IF_EQUAL(DestructiveInstall);
 	RETURN_IF_EQUAL(NonDestructiveInstall);
 	RETURN_IF_EQUAL(PrereqOnly);
+	RETURN_IF_EQUAL(Preload);
 	// Did not match
 	InstallMode = BuildPatchServices::EInstallMode::InvalidOrMax;
 #undef RETURN_IF_EQUAL
diff --git a/Engine/Source/Runtime/Online/BuildPatchServices/Public/BuildPatchMessage.h b/Engine/Source/Runtime/Online/BuildPatchServices/Public/BuildPatchMessage.h
--- a/Engine/Source/Runtime/Online/BuildPatchServices/Public/BuildPatchMessage.h
+++ b/Engine/Source/Runtime/Online/BuildPatchServices/Public/BuildPatchMessage.h
@@ -20,6 +20,11 @@ namespace BuildPatchServices
 			AccessLost = 0,
 			// Access has been regained after being lost.
 			AccessRegained,
+			// HTTP request failed
+			HttpRequestFailed,
+			Retired,
+
+			InvalidOrMax
 		};
 
 		// The type of event that occurred.
@@ -36,12 +41,16 @@ namespace BuildPatchServices
 		// Describes the action type.
 		enum class EType : uint32
 		{
-			// The file was removed.
-			Removed = 0,
-			// The file was added.
+			// The file was completed in the staging location.
+			Staged = 0,
+			// The file was added to the final install location.
 			Added,
-			// The file was updated.
+			// The file was updated in the final install location.
 			Updated,
+			// The file was removed from the final install location.
+			Removed,
+
+			InvalidOrMax
 		};
 
 		// The type of action that occurred.
@@ -56,7 +65,21 @@ namespace BuildPatchServices
 	struct FChunkUriRequest
 	{
 		FString CloudDirectory;
-		FString RelativePath;
+		FString ChunkCloudSubdirectory;
+		FString ChunkPath;
+		FString UniqueBuildId;
+		bool bChunkFromDelta = false;
+	};
+
+	/**
+	 * A Request for a delta file
+	 */
+	struct FDeltaUriRequest
+	{
+		FString CloudDirectory;
+		FString DeltaCloudSubdirectory;
+		FString DeltaPath;
+		FString UniqueBuildId;
 	};
 
 	/**
@@ -109,6 +132,24 @@ namespace BuildPatchServices
 	};
 	ENUM_CLASS_FLAGS(EMessageRequests);
 
+	/**
+	 * A Response containing the actual location of the file and URI's expiration time
+	 */
+	struct FExpiratedUriResponse
+	{
+		FString Uri;
+		FString SignatureExpiration;
+	};
+
+	/**
+	 * A Response containing an array of expirated URIs (e.g. different CDNs)
+	 */
+	struct FExpiratedUrisResponse
+	{
+		TArray<FExpiratedUriResponse> ExpiratedUri;
+		bool bNeedRetry = true;
+	};
+
 	/**
 	 * Base class of a message handler, this should be inherited from and passed to an installer to receive messages that you want to handle.
 	 */
@@ -144,9 +185,17 @@ namespace BuildPatchServices
 		 * Handles responding to a chunk Uri request
 		 * @param Request   The request for a chunk
 		 * @param OnResponse   The function to callback once the chunk has been found
+		 * @return true if request is handled successfully
 		 */
 		virtual bool HandleRequest(const FChunkUriRequest& Request, TFunction<void(FChunkUriResponse)> OnResponse) { return false; }
 
+		/**
+		 * Handles responding to a delta Uri request
+		 * @param Request   The request for a delta Uri
+		 * @param OnResponse   The function to callback once the Uri has been found
+		 */
+		virtual bool HandleRequest(const FDeltaUriRequest& Request, TFunction<void(FExpiratedUrisResponse)> OnResponse) { return false; }
+
 		/**
 		 * @return the message request flags.
 		 */
@@ -170,9 +219,85 @@ namespace BuildPatchServices
 		virtual bool HandleRequest(const FChunkUriRequest& Request, TFunction<void(FChunkUriResponse)> OnResponse) override
 		{
 			FChunkUriResponse Response;
-			Response.Uri = Request.CloudDirectory / Request.RelativePath;
+			Response.Uri = Request.CloudDirectory / Request.ChunkCloudSubdirectory / Request.ChunkPath;
+			OnResponse(MoveTemp(Response));
+			return true;
+		}
+
+		virtual bool HandleRequest(const FDeltaUriRequest& Request, TFunction<void(FExpiratedUrisResponse)> OnResponse) override
+		{
+			FExpiratedUrisResponse Response;
+			// Expiration time is not needed in the default case, setting this value explicitly empty
+			Response.ExpiratedUri.Add({ Request.CloudDirectory / Request.DeltaCloudSubdirectory / Request.DeltaPath, "" });		
 			OnResponse(MoveTemp(Response));
 			return true;
 		}
 	};
 }
+
+static_assert((uint32)BuildPatchServices::FChunkSourceEvent::EType::InvalidOrMax == 4, "Please add support for the extra values to the Lex functions below.");
+inline const TCHAR* LexToString(BuildPatchServices::FChunkSourceEvent::EType ChunkSourceEventType)
+{
+#define CASE_ENUM_TO_STR(Value) case BuildPatchServices::FChunkSourceEvent::EType::Value: return TEXT(#Value)
+	switch (ChunkSourceEventType)
+	{
+		CASE_ENUM_TO_STR(AccessLost);
+		CASE_ENUM_TO_STR(AccessRegained);
+		CASE_ENUM_TO_STR(HttpRequestFailed);
+		CASE_ENUM_TO_STR(Retired);
+		default: return TEXT("InvalidOrMax");
+	}
+#undef CASE_ENUM_TO_STR
+}
+
+inline void LexFromString(BuildPatchServices::FChunkSourceEvent::EType& ChunkSourceEventType, const TCHAR* Buffer)
+{
+#define RETURN_IF_EQUAL(Value) if (FCString::Stricmp(Buffer, TEXT(#Value)) == 0) { ChunkSourceEventType = BuildPatchServices::FChunkSourceEvent::EType::Value; return; }
+	const TCHAR* const Prefix = TEXT("FChunkSourceEvent::EType::");
+	const SIZE_T PrefixLen = FCString::Strlen(Prefix);
+	if (FCString::Strnicmp(Buffer, Prefix, PrefixLen) == 0)
+	{
+		Buffer += PrefixLen;
+	}
+	RETURN_IF_EQUAL(AccessLost);
+	RETURN_IF_EQUAL(AccessRegained);
+	RETURN_IF_EQUAL(HttpRequestFailed);
+	RETURN_IF_EQUAL(Retired);
+	// Did not match
+	ChunkSourceEventType = BuildPatchServices::FChunkSourceEvent::EType::InvalidOrMax;
+#undef RETURN_IF_EQUAL
+}
+
+inline const TCHAR* LexToString(BuildPatchServices::FInstallationFileAction::EType InstallationFileActionType)
+{
+#define CASE_ENUM_TO_STR(Value) case BuildPatchServices::FInstallationFileAction::EType::Value: return TEXT(#Value)
+	static_assert((uint32)BuildPatchServices::FInstallationFileAction::EType::InvalidOrMax == 4, "Please add support for the extra values below.");
+	switch (InstallationFileActionType)
+	{
+		CASE_ENUM_TO_STR(Staged);
+		CASE_ENUM_TO_STR(Added);
+		CASE_ENUM_TO_STR(Updated);
+		CASE_ENUM_TO_STR(Removed);
+		default: return TEXT("InvalidOrMax");
+	}
+#undef CASE_ENUM_TO_STR
+}
+
+inline void LexFromString(BuildPatchServices::FInstallationFileAction::EType& InstallationFileActionType, const TCHAR* Buffer)
+{
+#define RETURN_IF_EQUAL(Value) if (FCString::Stricmp(Buffer, TEXT(#Value)) == 0) { InstallationFileActionType = BuildPatchServices::FInstallationFileAction::EType::Value; return; }
+	const TCHAR* const Prefix = TEXT("FInstallationFileAction::EType::");
+	const SIZE_T PrefixLen = FCString::Strlen(Prefix);
+	if (FCString::Strnicmp(Buffer, Prefix, PrefixLen) == 0)
+	{
+		Buffer += PrefixLen;
+	}
+	static_assert((uint32)BuildPatchServices::FInstallationFileAction::EType::InvalidOrMax == 4, "Please add support for the extra values below.");
+	RETURN_IF_EQUAL(Staged);
+	RETURN_IF_EQUAL(Added);
+	RETURN_IF_EQUAL(Updated);
+	RETURN_IF_EQUAL(Removed);
+	// Did not match
+	InstallationFileActionType = BuildPatchServices::FInstallationFileAction::EType::InvalidOrMax;
+#undef RETURN_IF_EQUAL
+}
diff --git a/Engine/Source/Runtime/Online/BuildPatchServices/Public/BuildPatchSettings.h b/Engine/Source/Runtime/Online/BuildPatchServices/Public/BuildPatchSettings.h
--- a/Engine/Source/Runtime/Online/BuildPatchServices/Public/BuildPatchSettings.h
+++ b/Engine/Source/Runtime/Online/BuildPatchServices/Public/BuildPatchSettings.h
@@ -4,7 +4,9 @@
 #include "BuildPatchDelta.h"
 #include "BuildPatchFeatureLevel.h"
 #include "BuildPatchInstall.h"
+#include "BuildPatchMessage.h"
 #include "BuildPatchVerify.h"
+#include "BuildPatchState.h"
 #include "Containers/Array.h"
 #include "Containers/Map.h"
 #include "Containers/Set.h"
@@ -14,6 +16,7 @@
 #include "HAL/Platform.h"
 #include "Interfaces/IBuildManifest.h"
 #include "Interfaces/IBuildInstallerSharedContext.h"
+#include "Misc/Optional.h"
 #include "Misc/Variant.h"
 #include "Templates/UnrealTemplate.h"
 
@@ -23,6 +26,20 @@ namespace BuildPatchServices
 {
 	enum class EInstallActionIntent : int32;
 
+	// A delegate returning a bool.
+	DECLARE_DELEGATE_RetVal(bool, FChunkBuildBooleanDelegate);
+	DECLARE_DELEGATE_RetVal_OneParam(bool, FChunkBuildFilesEnumeratedDelegate, const TArray<FString>& /* Filenames */);
+	/**
+	 * An enum defining the method of sorting files in a build.
+	 */
+	enum class EFileSortOrder : uint8
+	{
+		AlphabeticalFilename = 0,
+		AlphabeticalTagThenFilename,
+
+		InvalidOrMax
+	};
+
 	/**
 	 * Defines a list of all build patch services initialization settings, can be used to override default init behaviors.
 	 */
@@ -69,8 +86,8 @@ namespace BuildPatchServices
 		static BUILDPATCHSERVICES_API FInstallerAction MakeUpdate(const IBuildManifestRef& CurrentManifest, const IBuildManifestRef& InstallManifest, TSet<FString> InstallTags = TSet<FString>(), FString InstallSubdirectory = FString(), FString CloudSubdirectory = FString());
 
 		/**
-		 * Creates an install action.
-		 * @param Manifest          The manifest for the build to be installed.
+		 * Creates a verify and repair action.
+		 * @param Manifest          The manifest for the build being repaired.
 		 * @param InstallTags       The install tags to use if selectively installing files. If empty set, all files will be repaired.
 		 * @param InstallSubdirectory  The subdirectory to install this data to inside of the main install directory.
 		 * @param CloudSubdirectory    The subdirectory of the cloud distribution root that this patch data should be sourced.
@@ -78,6 +95,16 @@ namespace BuildPatchServices
 		 */
 		static BUILDPATCHSERVICES_API FInstallerAction MakeRepair(const IBuildManifestRef& Manifest, TSet<FString> InstallTags = TSet<FString>(), FString InstallSubdirectory = FString(), FString CloudSubdirectory = FString());
 
+		/**
+		 * Creates a verify only action. Note that MakeRepair should be used usually, which will verify and try to repair any issues right away.
+		 * This action will only check integrity and fail or succeed based on results.
+		 * @param Manifest          The manifest for the build being verified.
+		 * @param InstallTags       The install tags to use if selectively installing files. If empty set, all files will be repaired.
+		 * @param InstallSubdirectory  The subdirectory to install this data to inside of the main install directory.
+		 * @param CloudSubdirectory    The subdirectory of the cloud distribution root that this patch data should be sourced.
+		 * @return the action setup forcing an SHA check, and repair of all tagged files.
+		 */
+		static BUILDPATCHSERVICES_API FInstallerAction MakeVerifyOnly(const IBuildManifestRef& Manifest, TSet<FString> InstallTags = TSet<FString>(), FString InstallSubdirectory = FString(), FString CloudSubdirectory = FString());
 
 		/**
 		 * Creates an uninstall action.
@@ -109,16 +136,6 @@ namespace BuildPatchServices
 			}
 		}
 
-		/**
-		 * Copy constructor.
-		 */
-		BUILDPATCHSERVICES_API FInstallerAction(const FInstallerAction& CopyFrom);
-
-		/**
-		 * RValue constructor to allow move semantics.
-		 */
-		BUILDPATCHSERVICES_API FInstallerAction(FInstallerAction&& MoveFrom);
-
 	public:
 		/**
 		 * @return true if this action intent is to perform a fresh installation.
@@ -135,6 +152,11 @@ namespace BuildPatchServices
 		 */
 		BUILDPATCHSERVICES_API bool IsRepair() const;
 
+		/**
+		 * @return true if this action intent is to only verify an existing installation, without repair or modification.
+		 */
+		BUILDPATCHSERVICES_API bool IsVerifyOnly() const;
+
 		/**
 		 * @return true if this action intent is to uninstall an installation.
 		 */
@@ -172,7 +194,7 @@ namespace BuildPatchServices
 		 */
 		IBuildManifestPtr TryGetCurrentManifest() const
 		{
-			if (IsUpdate() || IsRepair() || IsUninstall())
+			if (IsUpdate() || IsRepair() || IsVerifyOnly() || IsUninstall())
 			{
 				return GetCurrentManifest();
 			}
@@ -181,11 +203,11 @@ namespace BuildPatchServices
 
 		/**
 		 * Helper for getting the install manifest if nullable is preferred.
-		 * @return true is this action intent it to perform a fresh installation.
+		 * @return the manifest for the new installation, if valid.
 		 */
 		IBuildManifestPtr TryGetInstallManifest() const
 		{
-			if (IsInstall() || IsUpdate() || IsRepair())
+			if (IsInstall() || IsUpdate() || IsRepair() || IsVerifyOnly())
 			{
 				return GetInstallManifest();
 			}
@@ -230,20 +252,20 @@ namespace BuildPatchServices
 		BUILDPATCHSERVICES_API FInstallerConfiguration(FInstallerConfiguration&& MoveFrom);
 
 	public:
-		IBuildManifestPtr CurrentManifest;
 		IBuildManifestRef InstallManifest;
+		IBuildManifestPtr CurrentManifest = nullptr;
 		FString InstallDirectory;
 		FString StagingDirectory;
 		FString BackupDirectory;
 		TArray<FString> ChunkDatabaseFiles;
 		TArray<FString> CloudDirectories;
 		TSet<FString> InstallTags;
-		EInstallMode InstallMode;
-		EVerifyMode VerifyMode;
-		EDeltaPolicy DeltaPolicy;
-		bool bIsRepair;
-		bool bRunRequiredPrereqs;
-		bool bAllowConcurrentExecution;
+		EInstallMode InstallMode = EInstallMode::NonDestructiveInstall;
+		EVerifyMode VerifyMode = EVerifyMode::ShaVerifyAllFiles;
+		EDeltaPolicy DeltaPolicy = EDeltaPolicy::Skip;
+		bool bIsRepair = false;
+		bool bRunRequiredPrereqs = true;
+		bool bAllowConcurrentExecution = false;
 	};
 
 	/**
@@ -271,6 +293,9 @@ namespace BuildPatchServices
 		TArray<FString> ChunkDatabaseFiles;
 		// If set, when we've read everything we're going to read from a chunkdb file, we delete it to minimize required disk space for install/patch.
 		bool bDeleteChunkDbFilesAfterUse = false;
+		// If set, as we use chunks from a chunkdb file, we reduce it's size on disk to minimize required disk space for install/patch. This is only effective if the ChunkDB files are produced
+		// with reverse chunk ordering.
+		bool bTruncateChunkDbFilesAsUsed = false;
 
 		// If set, don't actually do an installation, just figure out how much disk space is required if we are using entirely
 		// chunkdbs and we delete them as we go.
@@ -280,18 +305,33 @@ namespace BuildPatchServices
 		// other sources (i.e. chunkdbs or install directory for patches). If the chunk can not be serviced from here, it's a failure.
 		TArray<FString> CloudDirectories;
 
+		// Map of available encryption secrets.
+		TMap<FGuid, TArray<uint8>> EncryptionSecrets;
+
 		// The mode for installation.
-		EInstallMode InstallMode;
+		EInstallMode InstallMode = EInstallMode::NonDestructiveInstall;
 		// The mode for verification. This will be overridden for any files referenced by a repair action.
-		EVerifyMode VerifyMode;
+		EVerifyMode VerifyMode = EVerifyMode::ShaVerifyAllFiles;
 		// The policy to follow for requesting an optimised delta.
-		EDeltaPolicy DeltaPolicy;
+		EDeltaPolicy DeltaPolicy = EDeltaPolicy::Skip;
+		// Optional specification for non-standard optimized delta
+		FString DeltaFilenameTrailer;
 		// Whether to run the prerequisite installer provided if it hasn't been ran before on this machine.
-		bool bRunRequiredPrereqs;
+		bool bRunRequiredPrereqs = true;
+		// Whether to skip the prerequisite installer provided if it already ran.
+		bool bSkipPrereqIfAlreadyRan = true;
 		// Whether to allow this installation to run concurrently with any existing installations.
-		bool bAllowConcurrentExecution;
+		bool bAllowConcurrentExecution = false;
+		// Download rate limit in bytes per second. (0 means unlimited)
+		uint64 DownloadRateLimitBps = 0;
+		// Use raw filenames for staging files
+		bool bStageWithRawFilenames = false;
+		// If true, fail the installation if the manifest contains a symlink.
+		bool bRejectSymlinks = false;
+		// Whether to run the custom uninstall action or not
+		bool bSkipCustomUninstallAction = false;
 		// Whether to gather individual file operation statistics during install
-		bool bTrackFileOperations;
+		bool bTrackFileOperations = false;
 
 		// Whether the file constructor is allowed to spawn multiple threads for IO read/write servicing.
 		// Unset uses Engine.ini value [Portal.BuildPatch]ConstructorSpawnAdditionalIOThreads if present, FFileConstructorConfig::bDefaultSpawnAdditionalIOThreads otherwise.
@@ -335,28 +375,20 @@ namespace BuildPatchServices
 
 		// If set, no initial size check will be performed prior to starting the installation. Running out of disk space during installation is unchanged.
 		bool bSkipInitialDiskSizeCheck = false;
-		// Optional specification for non-standard optimized delta
-		FString DeltaFilenameTrailer;
 	};
 
 	/**
 	 * Defines a list of all options for the build chunking task.
 	 */
-	struct FChunkBuildConfiguration
+	struct FDirectoryChunkerConfiguration
 	{
-	public:
-		/**
-		 * Default constructor
-		 */
-		BUILDPATCHSERVICES_API FChunkBuildConfiguration();
-
 	public:
 		// The client feature level to output data for.
-		EFeatureLevel FeatureLevel;
+		EFeatureLevel FeatureLevel = EFeatureLevel::Latest;
 		// The directory to analyze.
 		FString RootDirectory;
 		// The ID of the app of this build.
-		uint32 AppId;
+		uint32 AppId = 0;
 		// The name of the app of this build.
 		FString AppName;
 		// The version string for this build.
@@ -371,6 +403,8 @@ namespace BuildPatchServices
 		FString IgnoreListFile;
 		// The path to a file containing a \r\n separated list of RootDirectory relative files followed by attribute keywords.
 		FString AttributeListFile;
+		// The order in which to stream the build files.
+		EFileSortOrder FileSortOrder = EFileSortOrder::AlphabeticalTagThenFilename;
 		// The set of identifiers which the prerequisites satisfy.
 		TSet<FString> PrereqIds;
 		// The display name of the prerequisites installer.
@@ -379,14 +413,26 @@ namespace BuildPatchServices
 		FString PrereqPath;
 		// The command line arguments for the prerequisites installer.
 		FString PrereqArgs;
+		// The path to the uninstall custom action executable.
+		FString UninstallActionPath;
+		// The command line arguments for the uninstall custom action executable.
+		FString UninstallActionArgs;
 		// The maximum age (in days) of existing data files which can be reused in this build.
-		float DataAgeThreshold;
+		float DataAgeThreshold = TNumericLimits<float>::Max();
 		// Indicates whether data age threshold should be honored. If false, ALL data files can be reused.
-		bool bShouldHonorReuseThreshold;
+		bool bShouldHonorReuseThreshold = false;
 		// The chunk window size to be used when saving out new data.
-		uint32 OutputChunkWindowSize;
+		uint32 OutputChunkWindowSize = 1024 * 1024;
 		// Indicates whether any window size chunks should be matched, rather than just out output window size.
-		bool bShouldMatchAnyWindowSize;
+		bool bShouldMatchAnyWindowSize = true;
+		// Indicates that if a chunk was already found in the cloud directory, to resave it anyway.
+		bool bResaveExistingChunks = false;
+		// Indicates that if a chunk was matched from a previous build, to resave it anyway.
+		bool bResaveKnownChunks = false;
+		// The encryption secret key to use for newly created chunk files, and the manifest itself.
+		TArray<uint8> EncryptionSecretKey;
+		// An ID that identifies the secret key above, used to know which secrets are required for each manifest & chunk.
+		FGuid EncryptionSecretId;
 		// Map of custom fields to add to the manifest.
 		TMap<FString, FVariant> CustomFields;
 		// The cloud directory that all patch data will be saved to. An empty value will use module's global setting.
@@ -394,11 +440,52 @@ namespace BuildPatchServices
 		// The output manifest filename.
 		FString OutputFilename;
 		// Allow Manifest Creation for builds with no data
-		bool bAllowEmptyBuild;
+		bool bAllowEmptyBuild = false;
+		// Whether to follow symlinks, or if false, replicate them. True would process the build as if all links were actual files and directories.
+		// Note that this currently only affects mac and linux systems.
+		bool bFollowSymlinks = false;
+		// SHA256 and MD5 are expensive and can slow down the chunking process, set this true to skip calculation and save as zeroed values.
+		bool bSHA1Only = false;
+	};
+
+	// Temporary for use with deprecated module functions.
+	typedef FDirectoryChunkerConfiguration FGenerationConfiguration;
+	typedef FDirectoryChunkerConfiguration FChunkBuildConfiguration;
+
+	/**
+	 * An enum defining the units of DiffAbortThreshold value.
+	 */
+	enum class EDiffAbortThresholdUnits : uint8
+	{
+		Absolute = 0,
+		Percentage,
 	};
 
-	// Temporary for use with deprecated module function.
-	typedef FChunkBuildConfiguration FGenerationConfiguration;
+	/**
+	 * Defines the threshold for the original delta size, for which we would abort and not process.
+	 */
+	struct FDiffAbortThreshold
+	{
+	public:
+
+		// A threshold for the original delta size, for which we would abort and not process.
+		uint64 Value = 0ULL;
+
+		// A units of measurement for Value field. 
+		EDiffAbortThresholdUnits Unit = EDiffAbortThresholdUnits::Absolute;
+	};
+
+	/**
+	 * Defines a namespace for threshold value limits.
+	 */
+	namespace DiffAbortThresholdLimits
+	{
+		// Minimum threshold value limit, expressed as a percentage
+		constexpr uint64 MinPercentage = 1ULL;
+
+		// Minimum threshold value limit, expressed in absolute units.
+		constexpr uint64 MinAbsolute = 1000ULL * 1000ULL * 1000ULL;
+	}
 
 	/**
 	 * Defines a list of all options for the chunk delta optimisation task.
@@ -417,13 +504,37 @@ namespace BuildPatchServices
 		// A full file or http path for the manifest to be used as the destination build.
 		FString ManifestBUri;
 		// The cloud directory that all patch data will be saved to. An empty value will use ManifestB's directory.
-		FString CloudDirectory;
+		FString OutputCloudDirectory;
+		// The cloud directories that existing patch data can be downloaded from. An empty array will use ManifestB's directory.
+		TArray<FString> DownloadCloudDirectories;
 		// The window size to use for find new matches.
 		uint32 ScanWindowSize;
 		// The chunk size to use for saving new diff data.
 		uint32 OutputChunkSize;
+		// Indicates that if a chunk was already found in the cloud directory, to resave it anyway.
+		bool bResaveExistingChunks;
 		// A threshold for the original delta size, for which we would abort and not process.
-		uint64 DiffAbortThreshold;
+		TOptional<FDiffAbortThreshold> DiffAbortThreshold;
+		// Map of available encryption secrets, if any are missing, the process will fail.
+		TMap<FGuid, TArray<uint8>> EncryptionSecrets;
+		// Array of message handlers to provide signed URLs for source build chunks downloads.
+		TArray<FMessageHandler*> MessageHandlersA;
+		// Array of message handlers to provide signed URLs for destination build chunks and delta file downloads.
+		TArray<FMessageHandler*> MessageHandlersB;
+
+		// The old and new build aliases. Tags in the old and new build will be wholesale replaced by any aliases
+		// in these maps, allowing for control over what data is available for the patching system to match from. By default
+		// matches are only allowed from the same tag, with this one can virtually replace a set of tags with the same tag, allowing
+		// the patcher to select matches from all of the aliased tags.
+		TMap<FString, FString> ManifestATagAliases;
+		TMap<FString, FString> ManifestBTagAliases;
+		
+		// The custom optimized delta filename to use for generating this optimized delta. If used, all BPS modes
+		// that request an optimized delta will need the same trailer passed in.
+		FString DeltaFilenameTrailer;
+
+		// If specified, the source information for where each chunk matches from will be output to this file.
+		FString SourceDetailsLogFilename;
 	};
 
 	/**
@@ -451,39 +562,44 @@ namespace BuildPatchServices
 	 */
 	struct FDiffManifestsConfiguration
 	{
-	public:
-		/**
-		 * Default constructor
-		 */
-		BUILDPATCHSERVICES_API FDiffManifestsConfiguration();
-
 	public:
 		// A full file or http path for the manifest to be used as the source build.
 		FString ManifestAUri;
 		// A full file or http path for the manifest to be used as the destination build.
 		FString ManifestBUri;
+		// List of directories or http paths that can be used to request patch data. If empty, then path to ManifestA will be used.
+		TArray<FString> CloudDirs;
 		// The tag set to use to filter desired files from ManifestA.
 		TSet<FString> TagSetA;
 		// The tag set to use to filter desired files from ManifestB.
 		TSet<FString> TagSetB;
 		// Tag sets that will be used to calculate additional differential size statistics between manifests.
-		// They must all be a subset of anything used in TagSetB.
-		TArray<TSet<FString>> CompareTagSets;
+		// They must all be a subset of anything used in TagSetA.
+		TArray<TSet<FString>> CompareTagSetsA;
+		// If set, the tagset comparison will use these as the new tags and CompareTagSets will be the old installed tags, to allow for
+		// comparing the diff when the installed set of tags changes.
+		TArray<TSet<FString>> CompareTagSetsB;
 		// A full file path where a JSON object will be saved for the diff details.Empty string if not desired.
 		FString OutputFilePath;
+		// Array of message handlers to provide signed URIs for delta file.
+		TArray<FMessageHandler*> MessageHandlers;
 
 		// A destination base directory to use for writing out patch description files for the patch between the two
 		// manifest.
 		FString OutputPatchDescriptorPath;
 
-		// If true, the tool will throw an error if the optimized delta file is missing.
-		bool bRequireOptimizedDelta=false;
+		// By default, DiffManifest will use TryFetchContinueWithout and will, as a result, show large patches. If your workflows are such that
+		// you always expect an optimized delta, set this to cause DiffManifest to fail instead of show irrelevant data.
+		bool bRequireOptimizedDelta = false;
 
 		// If true, only emit patch descriptors and not the entirety of the diff manifests output.
-		bool bOnlyPatchDescriptors=false;
-		
+		bool bOnlyPatchDescriptors = false;
+
 		// If true, emit simulated installation time estimations
-		bool bEmitInstallTime=true;
+		bool bSimulateInstallTimes = false;
+
+		// The custom optimized delta filename to use, if desired.
+		FString DeltaFilenameTrailer;
 	};
 
 	/**
@@ -538,6 +654,101 @@ namespace BuildPatchServices
 		uint64 MaxOutputFileSize;
 		// A full file path to use when saving the json output data.
 		FString ResultDataFilePath;
+		// Optional flag for ignoring manifest file invalid tags
+		bool bIgnoreManifestFileInvalidTags;
+		// Optional flag for ignoring manifest describing a previous build file invalid tags
+		bool bIgnorePrevManifestFileInvalidTags;
+
+		// Optional specification for non-standard optimized delta
+		FString DeltaFilenameTrailer;
+	};
+
+	/**
+	 * Defines a list of all options for the chunk harvesting task.
+	 */
+	struct FChunkHarvesterConfiguration
+	{
+	public:
+		/**
+		 * Default constructor
+		 */
+		BUILDPATCHSERVICES_API FChunkHarvesterConfiguration();
+
+	public:
+		// Set an override feature level to output newer format.
+		EFeatureLevel FeatureLevelOverride;
+		// A full file path to the manifest to enumerate chunks from.
+		FString ManifestFilePath;
+		// The directory containing the binary image to be read.
+		FString BuildRoot;
+		// The directory where existing data will be recognized from, and new data added to.
+		FString CloudDir;
+		// Indicates that if a chunk was already found in the cloud directory, to resave it anyway.
+		bool bResaveExistingChunks;
 	};
 }
 
+static_assert((uint32)BuildPatchServices::EFileSortOrder::InvalidOrMax == 2, "Please add support for the extra values to the Lex functions below.");
+
+inline const TCHAR* LexToString(BuildPatchServices::EFileSortOrder FileSortOrder)
+{
+#define CASE_ENUM_TO_STR(Value) case BuildPatchServices::EFileSortOrder::Value: return TEXT(#Value)
+	switch (FileSortOrder)
+	{
+		CASE_ENUM_TO_STR(AlphabeticalFilename);
+		CASE_ENUM_TO_STR(AlphabeticalTagThenFilename);
+		default: return TEXT("InvalidOrMax");
+	}
+#undef CASE_ENUM_TO_STR
+}
+
+inline void LexFromString(BuildPatchServices::EFileSortOrder& FileSortOrder, const TCHAR* Buffer)
+{
+#define RETURN_IF_EQUAL(Value) if (FCString::Stricmp(Buffer, TEXT(#Value)) == 0) { FileSortOrder = BuildPatchServices::EFileSortOrder::Value; return; }
+	const TCHAR* const Prefix = TEXT("EFileSortOrder::");
+	const SIZE_T PrefixLen = FCString::Strlen(Prefix);
+	if (FCString::Strnicmp(Buffer, Prefix, PrefixLen) == 0)
+	{
+		Buffer += PrefixLen;
+	}
+	RETURN_IF_EQUAL(AlphabeticalFilename);
+	RETURN_IF_EQUAL(AlphabeticalTagThenFilename);
+	// Did not match
+	FileSortOrder = BuildPatchServices::EFileSortOrder::InvalidOrMax;
+#undef RETURN_IF_EQUAL
+}
+inline FString LexToString(const TOptional<BuildPatchServices::FDiffAbortThreshold>& DiffAbortThreshold)
+{
+	constexpr auto* DiffAbortThresholdString = TEXT("Never Abort");
+	if (!DiffAbortThreshold.IsSet())
+	{
+		return DiffAbortThresholdString;
+	}
+
+	switch (DiffAbortThreshold.GetValue().Unit)
+	{
+	case BuildPatchServices::EDiffAbortThresholdUnits::Absolute:
+		return FString::FormatAsNumber(int32(DiffAbortThreshold.GetValue().Value / (1024 * 1024))) + TEXT(" MiB");
+	case BuildPatchServices::EDiffAbortThresholdUnits::Percentage:
+		return FString::FormatAsNumber(int32(DiffAbortThreshold.GetValue().Value)) + TEXT("%");
+	default:
+		checkf(false, TEXT("Unexpected EDiffAbortThresholdUnits value: %hhu"), DiffAbortThreshold.GetValue().Unit);
+		break;
+	}
+
+	return DiffAbortThresholdString;
+}
+
+inline FString LexToString(BuildPatchServices::EDiffAbortThresholdUnits Unit)
+{
+	#define CASE_ENUM_TO_STR(Value) case BuildPatchServices::EDiffAbortThresholdUnits::Value: return TEXT(#Value)
+
+	switch (Unit)
+	{
+		CASE_ENUM_TO_STR(Absolute);
+		CASE_ENUM_TO_STR(Percentage);
+		default: return FString::Printf(TEXT("Unexpected value: %hhu"), Unit);
+	}
+
+	#undef CASE_ENUM_TO_STR
+}
diff --git a/Engine/Source/Runtime/Online/BuildPatchServices/Public/BuildPatchState.h b/Engine/Source/Runtime/Online/BuildPatchServices/Public/BuildPatchState.h
--- a/Engine/Source/Runtime/Online/BuildPatchServices/Public/BuildPatchState.h
+++ b/Engine/Source/Runtime/Online/BuildPatchServices/Public/BuildPatchState.h
@@ -29,6 +29,9 @@ namespace BuildPatchServices
 		// The patch process is installing files
 		Installing,
 
+		// The patch process is running custom uninstall action
+		RunningCustomUninstallAction,
+
 		// The patch process is moving staged files to the install
 		MovingToInstall,
 
@@ -63,7 +66,7 @@ namespace BuildPatchServices
 	BUILDPATCHSERVICES_API const FText& StateToText(const EBuildPatchState& State);
 }
 
-static_assert((uint32)BuildPatchServices::EBuildPatchState::NUM_PROGRESS_STATES == 12, "Please add support for the extra values to the Lex functions below.");
+static_assert((uint32)BuildPatchServices::EBuildPatchState::NUM_PROGRESS_STATES == 13, "Please add support for the extra values to the Lex functions below.");
 
 inline const TCHAR* LexToString(BuildPatchServices::EBuildPatchState State)
 {
@@ -75,6 +78,7 @@ inline const TCHAR* LexToString(BuildPatchServices::EBuildPatchState State)
 		CASE_ENUM_TO_STR(Resuming);
 		CASE_ENUM_TO_STR(Downloading);
 		CASE_ENUM_TO_STR(Installing);
+		CASE_ENUM_TO_STR(RunningCustomUninstallAction);
 		CASE_ENUM_TO_STR(MovingToInstall);
 		CASE_ENUM_TO_STR(SettingAttributes);
 		CASE_ENUM_TO_STR(BuildVerification);
@@ -101,6 +105,7 @@ inline void LexFromString(BuildPatchServices::EBuildPatchState& Error, const TCH
 	RETURN_IF_EQUAL(Resuming);
 	RETURN_IF_EQUAL(Downloading);
 	RETURN_IF_EQUAL(Installing);
+	RETURN_IF_EQUAL(RunningCustomUninstallAction);
 	RETURN_IF_EQUAL(MovingToInstall);
 	RETURN_IF_EQUAL(SettingAttributes);
 	RETURN_IF_EQUAL(BuildVerification);
diff --git a/Engine/Source/Runtime/Online/BuildPatchServices/Public/Interfaces/IBuildInstaller.h b/Engine/Source/Runtime/Online/BuildPatchServices/Public/Interfaces/IBuildInstaller.h
--- a/Engine/Source/Runtime/Online/BuildPatchServices/Public/Interfaces/IBuildInstaller.h
+++ b/Engine/Source/Runtime/Online/BuildPatchServices/Public/Interfaces/IBuildInstaller.h
@@ -57,6 +57,9 @@ enum class EBuildPatchInstallError
 	// An error occurred creating a file due to their not being enough space left on the disk
 	OutOfDiskSpace = 11,
 
+	// The custom uninstall action failed
+	CustomUninstallActionError = 12,
+
 	// Used to help verify logic
 	NumInstallErrors
 };
@@ -80,6 +83,7 @@ namespace InstallErrorPrefixes
 		TEXT("IZ"), // InitializationError
 		TEXT("PL"), // PathLengthExceeded
 		TEXT("DS"), // OutOfDiskSpace
+		TEXT("UA"), // UninstallCustomActionError
 	};
 };
 
@@ -101,6 +105,21 @@ enum class EBuildPatchDownloadHealth
 	NUM_Values
 };
 
+/**
+ * Tracks install stats for a given cloud directory.
+ */
+struct FCloudDirectoryStats
+{
+	// The identifier for this cloud directory, e.g. base URL.
+	FString CloudDirectory;
+
+	// Total number of bytes downloaded from this cloud directory.
+	uint64 TotalReceivedData = 0;
+
+	// Total number of errors encountered from this cloud directory.
+	uint32 ErrorCount = 0;
+};
+
 /**
  * A struct to hold stats for the build process.
  */
@@ -140,6 +159,7 @@ struct FBuildInstallStats
 		, NumFilesConstructed(0)
 		, InitializeTime(0.0f)
 		, ConstructTime(0.0f)
+		, UninstallActionTime(0.0f)
 		, MoveFromStageTime(0.0f)
 		, FileAttributesTime(0.0f)
 		, VerifyTime(0.0f)
@@ -164,6 +184,7 @@ struct FBuildInstallStats
 		, DisconnectedDownloadHealthTime(0.0f)
 		, ProcessRequiredDiskSpace(0)
 		, ProcessAvailableDiskSpace(0)
+		, CloudDirectoryStats()
 	{}
 
 	// The total number of files in the build.
@@ -208,6 +229,12 @@ struct FBuildInstallStats
 	uint64 TotalDownloadedData;
 	// The peak number of simultaneous download requests
 	uint32 ActiveRequestCountPeak;
+	// The average number of simultaneous download requests 
+	uint32 ActiveRequestCountAvg;
+	// The Experiment name if applicable for A/B Testing
+	FString ExperimentName;
+	// The Variant of the Experiment if applicable for A/B Testing
+	FString ExerimentVariantName;
 	// The average chunk download speed.
 	double AverageDownloadSpeed;
 	// The peak chunk download speed.
@@ -239,6 +266,10 @@ struct FBuildInstallStats
 	float InitializeTime;
 	// The time spent during the construction stage.
 	float ConstructTime;
+	// The time spent on chunks decryption.
+	uint64 TotalDecryptTime;
+	// The time spent during the custom uninstall actiion run.
+	float UninstallActionTime;
 	// The time spent moving staged files into the installation location.
 	float MoveFromStageTime;
 	// The time spent during the file attribution stage.
@@ -291,6 +322,8 @@ struct FBuildInstallStats
 	uint64 ProcessRequiredDiskSpace;
 	// The total number of bytes available at the time of checking ProcessRequiredDiskSpace.
 	uint64 ProcessAvailableDiskSpace;
+	// Stats tracked per cloud directory.
+	TArray<FCloudDirectoryStats> CloudDirectoryStats;
 };
 
 namespace BuildPatchServices
@@ -345,6 +378,24 @@ public:
 	 */
 	virtual bool IsUpdate() const = 0;
 
+	/**
+	 * Get whether the install is performing a repair of an existing install.
+	 * @return	true if installation is a repair
+	 */
+	virtual bool IsRepair() const = 0;
+
+	/**
+	 * Get whether the install is performing a verify only of an existing install, without repair.
+	 * @return	true if installation is a verify only.
+	 */
+	virtual bool IsVerifyOnly() const = 0;
+
+	/**
+	 * Get whether the install is actually uninstalling instead.
+	 * @return	true if installation is actually an uninstallation
+	 */
+	virtual bool IsUninstall() const = 0;
+
 	/**
 	 * Get whether the process was successfully complete, meaning no errors, and was not cancelled. Only valid if complete.
 	 * @return	true if installation process completed successfully.
@@ -460,6 +511,17 @@ public:
 	 */
 	virtual const BuildPatchServices::FBuildInstallerConfiguration& GetConfiguration() const = 0;
 
+	/**
+	 * Sets the download rate limit in Bytes per second. This throttles the download rate to the specified value.
+	 * @param InDownloadRateLimitBps    throttle value in bytes per second. 
+	 */
+	virtual void SetDownloadBytesPerSecondLimit(uint64 InDownloadRateLimitBps) = 0;
+
+	/**
+	 * Gets the download rate limit in Bytes per second.
+	 * @return the current download rate limit.
+	 */
+	virtual uint64 GetDownloadBytesPerSecondLimit() const = 0;
 #if !UE_BUILD_SHIPPING
 	/**
 	 * Returns debug text for download screen to help troubleshoot download issues
@@ -483,7 +545,7 @@ public:
 	virtual void SetPerFileSubdirectories(TMap<FString, FString>&& InOldPerFileSubdirectories, TMap<FString, FString>&& InNewPerFileSubdirectories) = 0;
 };
 
-static_assert((uint32)EBuildPatchInstallError::NumInstallErrors == 12, "Please add support for the extra values to the Lex functions below.");
+static_assert((uint32)EBuildPatchInstallError::NumInstallErrors == 13, "Please add support for the extra values to the Lex functions below.");
 
 inline const TCHAR* LexToString(EBuildPatchInstallError Error)
 {
@@ -502,6 +564,7 @@ inline const TCHAR* LexToString(EBuildPatchInstallError Error)
 		CASE_ENUM_TO_STR(InitializationError);
 		CASE_ENUM_TO_STR(PathLengthExceeded);
 		CASE_ENUM_TO_STR(OutOfDiskSpace);
+		CASE_ENUM_TO_STR(CustomUninstallActionError);
 	default: return TEXT("InvalidOrMax");
 	}
 #undef CASE_ENUM_TO_STR
@@ -528,6 +591,7 @@ inline void LexFromString(EBuildPatchInstallError& Error, const TCHAR* Buffer)
 	RETURN_IF_EQUAL(InitializationError);
 	RETURN_IF_EQUAL(PathLengthExceeded);
 	RETURN_IF_EQUAL(OutOfDiskSpace);
+	RETURN_IF_EQUAL(CustomUninstallActionError);
 	// Did not match
 	Error = EBuildPatchInstallError::NumInstallErrors;
 	return;
diff --git a/Engine/Source/Runtime/Online/BuildPatchServices/Public/Interfaces/IBuildManifest.h b/Engine/Source/Runtime/Online/BuildPatchServices/Public/Interfaces/IBuildManifest.h
--- a/Engine/Source/Runtime/Online/BuildPatchServices/Public/Interfaces/IBuildManifest.h
+++ b/Engine/Source/Runtime/Online/BuildPatchServices/Public/Interfaces/IBuildManifest.h
@@ -9,6 +9,8 @@
 #pragma once
 
 #include "CoreMinimal.h"
+#include "BuildPatchFeatureLevel.h"
+#include "Misc/SecureHash.h"
 
 class IBuildManifest;
 class IManifestField;
@@ -18,6 +20,21 @@ typedef TSharedRef< class IManifestField, ESPMode::ThreadSafe > IManifestFieldRe
 typedef TSharedPtr< class IBuildManifest, ESPMode::ThreadSafe > IBuildManifestPtr;
 typedef TSharedRef< class IBuildManifest, ESPMode::ThreadSafe > IBuildManifestRef;
 
+/**
+ * Declares flags for meta or attributes associated with a file.
+ */
+enum class EFileMetaFlags : uint8
+{
+	None = 0,
+	// Flag for readonly file.
+	ReadOnly = 1,
+	// Flag for natively compressed.
+	Compressed = 1 << 1,
+	// Flag for unix executable.
+	UnixExecutable = 1 << 2
+};
+ENUM_CLASS_FLAGS(EFileMetaFlags);
+
 /**
  * Interface to a manifest field, which is used for accessing custom fields in the manifest
  */
@@ -58,7 +75,13 @@ public:
 	 * Virtual destructor.
 	 */
 	virtual ~IBuildManifest() = default;
-	
+
+	/**
+	 * Gets the manifest feature level
+	 * @return		the feature level this manifest was built with
+	 */
+	virtual BuildPatchServices::EFeatureLevel GetFeatureLevel() const = 0;
+
 	/**
 	 * Get the App ID that this manifest belongs to
 	 * @return		the app ID
@@ -76,6 +99,31 @@ public:
 	 * @return		the version string
 	 */
 	virtual const FString& GetVersionString() const = 0;
+
+	/**
+	 * Get the unique string generated when producing the manifest
+	 * @return		the unique build id string
+	 */
+	virtual const FString& GetUniqueBuildId() const = 0;
+
+	/**
+	 * If this manifest is or was encrypted, returns the GUID that identifies the secret key required for decryption.
+	 * @param OUT	If the manifest in memory is currently encrypted.
+	 * @return		the secret id
+	 */
+	virtual const FGuid& GetEncryptionSecretId(bool* bOutIsManifestEncrypted = nullptr) const = 0;
+
+	/**
+	 * Regardless of whether the manifest was ever encrypted, returns whether the manifest is currently encrypted.
+	 * @return		If the manifest in memory is currently encrypted
+	 */
+	virtual bool IsManifestEncrypted() const = 0;
+
+	/**
+	 * Get all necessary secret key IDs.
+	 * @return		full set of necessary secret ids
+	 */
+	virtual TSet<FGuid> GetNecessaryEncryptionSecretIds() const = 0;
 	
 	/**
 	 * Get the local install path to the exe that launches the App
@@ -113,6 +161,18 @@ public:
 	 */
 	virtual const FString& GetPrereqArgs() const = 0;
 
+	/**
+	 * Get the path to the custom uninstall action exe
+	 * @return		local path to the custom uninstall action exe
+	 */
+	virtual const FString& GetUninstallActionPath() const = 0;
+
+	/**
+	 * Get the command line arguments that should be passed to the custom uninstall action exe
+	 * @return		the custom uninstall action command line arguments
+	 */
+	virtual const FString& GetUninstallActionArgs() const = 0;
+
 	/**
 	 * Get the size of this download, assuming fresh install
 	 * @return		the total download size in bytes
@@ -205,13 +265,40 @@ public:
 	virtual int64 GetFileSize(const TSet<FString>& Filenames) const = 0;
 
 	/**
-	 * Gets the file hash for a given file
-	 * @param Filename		IN		The filename in the build
-	 * @param OutHash		OUT		Receives the hash value if found
-	 * @return	true if we had the hash for this file
+	 * Get the file meta flags assigned to a particular file in the build.
+	 * @param Filename      The file.
+	 * @return		the file meta flags
 	 */
-	virtual bool GetFileHash(const FString& Filename, FSHAHash& OutHash) const = 0;
-	
+	virtual EFileMetaFlags GetFileMetaFlags(const FString& Filename) const = 0;
+
+	/**
+	 * Gets the MD5 hash for a given file
+	 * @param Filename      The file.
+	 * @return ptr to the MD5 for the given file. Nullptr if not found.
+	 */
+	virtual const FMD5Hash* GetFileMD5Hash(const FString& Filename) const = 0;
+
+	/**
+	 * Gets the SHA1 hash for a given file
+	 * @param Filename      The file.
+	 * @return ptr to the SHA1 for the given file. Nullptr if not found.
+	 */
+	virtual const FSHAHash* GetFileSHA1Hash(const FString& Filename) const = 0;
+
+	/**
+	 * Gets the SHA256 hash for a given file
+	 * @param Filename      The file.
+	 * @return ptr to the SHA256 for the given file. Nullptr if not found.
+	 */
+	virtual const FSHA256Signature* GetFileSHA256Hash(const FString& Filename) const = 0;
+
+	/**
+	 * Gets the MIME type calculated for a given file
+	 * @param Filename      The file.
+	 * @return ptr to the MIME type string for the given file. Nullptr if not found.
+	 */
+	virtual const FString* GetFileMIMEType(const FString& Filename) const = 0;
+
 	/**
 	 * Get the list of install tags in this manifest
 	 * @return the tags referenced.
@@ -245,6 +332,36 @@ public:
 	 */
 	virtual void GetRemovableFiles(const TCHAR* InstallPath, TArray<FString>& RemovableFiles) const = 0;
 
+	/**
+	 * Provides the set of chunks required to produce the given files.
+	 * @param Filenames         IN      The set of files.
+	 * @param RequiredChunks    OUT     The set of chunk GUIDs needed for those files.
+	 */
+	virtual void GetChunksRequiredForFiles(const TSet<FString>& Filenames, TSet<FGuid>& RequiredChunks) const = 0;
+
+	/**
+	 * Returns the size of a particular data file by it's GUID.
+	 * @param DataGuid		The GUID for the data
+	 * @return		File size, or 0 if this data is not in the manifest.
+	 */
+	virtual int64 GetDataSize(const FGuid& DataGuid) const = 0;
+
+	/**
+	 * Gets the SHA1 hash for a given chunk
+	 * @param ChunkGuid		IN		The guid of the chunk to get hash for
+	 * @param OutHash		OUT		Receives the hash value if found
+	 * @return	true if we had the hash for this chunk
+	 */
+	virtual bool GetChunkShaHash(const FGuid& ChunkGuid, FSHAHash& OutHash) const = 0;
+
+	/**
+	 * Gets the secret ID of the chunk
+	 * @param ChunkGuid					IN		The guid of the chunk to get secret ID for
+	 * @param OutEncryptionSecretId		OUT		The guid value of chunk's secret ID if found
+	 * @return true the chunk has encryption secret
+	 */
+	virtual bool GetChunkEncryptionSecretId(const FGuid& ChunkGuid, FGuid& OutEncryptionSecretId) const = 0;
+
 	/**
 	 * Checks the manifest format version to see if this manifest was loaded from latest data
 	 * @return True if the manifest was created from the latest format
@@ -258,6 +375,12 @@ public:
 	 */
 	virtual void CopyCustomFields(const IBuildManifestRef& Other, bool bClobber) = 0;
 
+	/**
+	 * Get the list of custom field names in this manifest
+	 * @return an array containing custom field names
+	 */
+	virtual TArray<FString> GetCustomFieldNames() const = 0;
+
 	/**
 	 * Get a custom field from the manifest
 	 * @param	FieldName	The name of the custom field
diff --git a/Engine/Source/Runtime/Online/BuildPatchServices/Public/Interfaces/IBuildPatchServicesModule.h b/Engine/Source/Runtime/Online/BuildPatchServices/Public/Interfaces/IBuildPatchServicesModule.h
--- a/Engine/Source/Runtime/Online/BuildPatchServices/Public/Interfaces/IBuildPatchServicesModule.h
+++ b/Engine/Source/Runtime/Online/BuildPatchServices/Public/Interfaces/IBuildPatchServicesModule.h
@@ -5,7 +5,10 @@
 #include "Modules/ModuleInterface.h"
 #include "Interfaces/IBuildInstaller.h"
 #include "Interfaces/IBuildStatistics.h"
+#include "Interfaces/IChunkDeltaOptimiser.h"
+#include "Interfaces/IDirectoryChunker.h"
 #include "Interfaces/IPatchDataEnumeration.h"
+#include "Interfaces/IChunkHarvester.h"
 #include "BuildPatchSettings.h"
 
 class IAnalyticsProvider;
@@ -15,6 +18,8 @@ class IAnalyticsProvider;
  */
 DECLARE_DELEGATE_OneParam(FBuildPatchInstallerDelegate, const IBuildInstallerRef&);
 DECLARE_DELEGATE_TwoParams(FBuildPatchBoolManifestDelegate, bool, IBuildManifestRef);
+DECLARE_MULTICAST_DELEGATE_TwoParams(FSingleChunkWrittenDelegate, FString /* Filename */, FString /* MD5Hash */);
+DECLARE_MULTICAST_DELEGATE_ThreeParams(FManifestCreatedDelegate, FString /* Filename */, FString /* SHA1Hash */, FString /* MD5Hash */);
 
 namespace ECompactifyMode
 {
@@ -55,17 +60,38 @@ public:
 	/**
 	 * Factory providing construction of a build statistics class.
 	 * @param Installer     The installer to create a build statistics for.
-	 * @return an instance of an IBuildStatistics implementation. 
+	 * @return an instance of an IBuildStatistics implementation.
 	 */
 	virtual BuildPatchServices::IBuildStatisticsRef CreateBuildStatistics(const IBuildInstallerRef& Installer) const = 0;
 
+	/**
+	 * Factory providing construction of a directory chunker class.
+	 * @param Configuration         Specifies the settings for the operation.  See BuildPatchServices::FDirectoryChunkerConfiguration comments.
+	 * @return an instance of an IDirectoryChunker implementation.
+	 */
+	virtual BuildPatchServices::IDirectoryChunkerRef CreateDirectoryChunker(BuildPatchServices::FDirectoryChunkerConfiguration Configuration) const = 0;
+
 	/**
 	 * Factory providing construction of a patch data enumeration class.
 	 * @param Configuration         Specifies the settings for the operation.  See BuildPatchServices::FPatchDataEnumerationConfiguration comments.
-	 * @return an instance of an IPatchDataEnumeration implementation. 
+	 * @return an instance of an IPatchDataEnumeration implementation.
 	 */
 	virtual BuildPatchServices::IPatchDataEnumerationRef CreatePatchDataEnumeration(BuildPatchServices::FPatchDataEnumerationConfiguration Configuration) const = 0;
 
+	/**
+	 * Factory providing construction of a chunk delta optimiser class.
+	 * @param Configuration         Specifies the settings for the operation.  See BuildPatchServices::FChunkDeltaOptimiserConfiguration comments.
+	 * @return an instance of an IChunkDeltaOptimiser implementation.
+	 */
+	virtual BuildPatchServices::IChunkDeltaOptimiserRef CreateChunkDeltaOptimiser(BuildPatchServices::FChunkDeltaOptimiserConfiguration Configuration) const = 0;
+
+	/**
+	 * Factory providing construction of a chunk harvester class.
+	 * @param Configuration         Specifies the settings for the operation.  See BuildPatchServices::FChunkHarvesterConfiguration comments.
+	 * @return an instance of an IChunkHarvester implementation.
+	 */
+	virtual BuildPatchServices::IChunkHarvesterRef CreateChunkHarvester(BuildPatchServices::FChunkHarvesterConfiguration Configuration) const = 0;
+
 	/**
 	 * Loads a Build Manifest from file and returns the interface
 	 * @param Filename		The file to load from
@@ -146,26 +172,31 @@ public:
 	virtual bool UnregisterAppInstallation(const FString AppInstallDirectory) = 0;
 
 	/**
-	 * Call to force the exit out of all current installers, optionally blocks until threads have exited and complete delegates are called.
-	 * @param WaitForThreads		If true, will block on threads exit and completion delegates
+	 * Registers an encryption secret. This information is used by installers to decrypt patch data that references this secret.
+	 * @param EncryptionSecretId    The ID for the secret key.
+	 * @param EncryptionSecretKey   The secret key.
 	 */
-	virtual void CancelAllInstallers(bool WaitForThreads) = 0;
+	virtual void RegisterEncryptionSecret(const FGuid& EncryptionSecretId, TArray<uint8> EncryptionSecretKey) = 0;
 
 	/**
-	 * Processes a Build directory to create chunks for new data and produce a manifest, saved to the provided cloud directory.
-	 * NOTE: This function is blocking and will not return until finished.
-	 * @param Configuration         Specifies the settings for the operation. See BuildPatchServices::FChunkBuildConfiguration comments.
-	 * @return true if successful.
+	 * Unregisters an encryption secret.
+	 * @param EncryptionSecretId    The ID for the secret key.
+	 * @return returns true if the secret existed and was removed.
 	 */
-	virtual bool ChunkBuildDirectory(const BuildPatchServices::FChunkBuildConfiguration& Configuration) = 0;
+	virtual bool UnregisterEncryptionSecret(const FGuid& EncryptionSecretId) = 0;
 
 	/**
-	 * Process a pair of manifests to produce additional delta data which reduces the patch directly between them.
-	 * NOTE: This function is blocking and will not return until finished.
-	 * @param Configuration         Specifies the settings for the operation. See BuildPatchServices::FChunkDeltaOptimiserConfiguration comments.
-	 * @return true if successful.
+	 * Generates random data via OpenSSL.
+	 * @param OutData    The object to fill with random data.
+	 * @return returns true if OpenSSL was able to generate random data.
+	*/
+	virtual bool CreateRandomBytes(const TArrayView<uint8> OutData) = 0;
+
+	/**
+	 * Call to force the exit out of all current installers, optionally blocks until threads have exited and complete delegates are called.
+	 * @param WaitForThreads		If true, will block on threads exit and completion delegates
 	 */
-	virtual bool OptimiseChunkDelta(const BuildPatchServices::FChunkDeltaOptimiserConfiguration& Configuration) = 0;
+	virtual void CancelAllInstallers(bool WaitForThreads) = 0;
 
 	/**
 	 * Processes a Cloud Directory to identify and delete any orphaned chunks or files.
@@ -175,14 +206,6 @@ public:
 	 */
 	virtual bool CompactifyCloudDirectory(const BuildPatchServices::FCompactifyConfiguration& Configuration) = 0;
 
-	/**
-	 * Saves info for an enumeration of patch data referenced from an input file of known format, to a specified output file.
-	 * NOTE: THIS function is blocking and will not return until finished.
-	 * @param Configuration         Specifies the settings for the operation. See BuildPatchServices::FPatchDataEnumerationConfiguration comments.
-	 * @return true if successful.
-	 */
-	virtual bool EnumeratePatchData(const BuildPatchServices::FPatchDataEnumerationConfiguration& Configuration) = 0;
-
 	/**
 	 * Searches a given directory for chunk and chunkdb files, and verifies their integrity uses the hashes in the files.
 	 * NOTE: THIS function is blocking and will not return until finished. Don't run on main thread.
@@ -215,6 +238,18 @@ public:
 	 */
 	virtual bool MergeManifests(const FString& ManifestFilePathA, const FString& ManifestFilePathB, const FString& ManifestFilePathC, const FString& NewVersionString, const FString& SelectionDetailFilePath) = 0;
 
+	/**
+	 * Takes two manifests as input, in order to merge together producing a new manifest containing all files.
+	 * NOTE: THIS function is blocking and will not return until finished. Don't run on main thread.
+	 * @param ManifestFilePathA         A full file path for the base manifest to be loaded.
+	 * @param ManifestFilePathB         A full file path for the merge manifest to be loaded, by default files in B will stomp over A.
+	 * @param NewVersionString          The new version string for the build, all other meta will be copied from B.
+	 * @param FilesFromA                A set containing the files that should be kept from ManifestA.
+	 * @param FilesFromB                A set containing the files that should be kept from ManifestB.
+	 * @return resulting manifest if successful.
+	 */
+	virtual IBuildManifestPtr MergeManifests(const IBuildManifestRef& ManifestA, const IBuildManifestRef& ManifestB, const FString& NewVersionString, const TSet<FString>& FilesFromA, const TSet<FString>& FilesFromB) = 0;
+
 	/**
 	 * Takes two manifests as input and outputs the details of the patch.
 	 * NOTE: THIS function is blocking and will not return until finished. Don't run on main thread.
@@ -235,10 +270,43 @@ public:
 	UE_DEPRECATED(4.21, "MakeManifestFromJSON(const FString& ManifestJSON) has been deprecated.  Please use MakeManifestFromData(const TArray<uint8>& ManifestData) instead.")
 	virtual IBuildManifestPtr MakeManifestFromJSON(const FString& ManifestJSON) = 0;
 
-	UE_DEPRECATED(4.21, "Please use ChunkBuildDirectory instead.")
+	UE_DEPRECATED(4.16, "EnumerateManifestData(const FString& ManifestFilePath, const FString& OutputFile, bool bIncludeSizes) has been deprecated.  Please use CreatePatchDataEnumeration(BuildPatchServices::FPatchDataEnumerationConfiguration Configuration) instead.")
+	virtual bool EnumerateManifestData(const FString& ManifestFilePath, const FString& OutputFile, bool bIncludeSizes)
+	{
+		BuildPatchServices::FPatchDataEnumerationConfiguration Configuration;
+		Configuration.InputFile = ManifestFilePath;
+		Configuration.OutputFile = OutputFile;
+		Configuration.bIncludeSizes = bIncludeSizes;
+		BuildPatchServices::IPatchDataEnumerationRef PatchDataEnumeration = CreatePatchDataEnumeration(MoveTemp(Configuration));
+		return PatchDataEnumeration->Run();
+	}
+
+	UE_DEPRECATED(4.23, "ChunkBuildDirectory(const BuildPatchServices::FChunkBuildConfiguration& Configuration) has been deprecated.  Please use CreateDirectoryChunker(BuildPatchServices::FDirectoryChunkerConfiguration Configuration) instead.")
+	virtual bool ChunkBuildDirectory(const BuildPatchServices::FChunkBuildConfiguration& Configuration)
+	{
+		BuildPatchServices::IDirectoryChunkerRef DirectoryChunker = CreateDirectoryChunker(Configuration);
+		return DirectoryChunker->Run();
+	}
+
+	UE_DEPRECATED(4.23, "EnumeratePatchData(const BuildPatchServices::FPatchDataEnumerationConfiguration& Configuration) has been deprecated.  Please use CreatePatchDataEnumeration(BuildPatchServices::FPatchDataEnumerationConfiguration Configuration) instead.")
+	virtual bool EnumeratePatchData(const BuildPatchServices::FPatchDataEnumerationConfiguration& Configuration)
+	{
+		BuildPatchServices::IPatchDataEnumerationRef PatchDataEnumeration = CreatePatchDataEnumeration(Configuration);
+		return PatchDataEnumeration->Run();
+	}
+
+	UE_DEPRECATED(4.23, "OptimiseChunkDelta(const BuildPatchServices::FChunkDeltaOptimiserConfiguration& Configuration) has been deprecated.  Please use CreateChunkDeltaOptimiser(BuildPatchServices::FChunkDeltaOptimiserConfiguration Configuration) instead.")
+	virtual bool OptimiseChunkDelta(const BuildPatchServices::FChunkDeltaOptimiserConfiguration& Configuration)
+	{
+		BuildPatchServices::IChunkDeltaOptimiserRef ChunkDeltaOptimiser = CreateChunkDeltaOptimiser(Configuration);
+		return ChunkDeltaOptimiser->Run();
+	}
+
+	UE_DEPRECATED(4.21, "GenerateChunksManifestFromDirectory(const BuildPatchServices::FGenerationConfiguration& Configuration) has been deprecated.  Please use CreateDirectoryChunker(BuildPatchServices::FDirectoryChunkerConfiguration Configuration) instead.")
 	virtual bool GenerateChunksManifestFromDirectory(const BuildPatchServices::FGenerationConfiguration& Configuration)
 	{
-		return ChunkBuildDirectory(Configuration);
+		BuildPatchServices::IDirectoryChunkerRef DirectoryChunker = CreateDirectoryChunker(Configuration);
+		return DirectoryChunker->Run();
 	}
 
 	UE_DEPRECATED(4.23, "StartBuildInstall(IBuildManifestPtr, IBuildManifestPtr, const FString&, FBuildPatchBoolManifestDelegate, bool, TSet<FString>) has been deprecated.  Please use CreateBuildInstaller(BuildPatchServices::FBuildInstallerConfiguration, FBuildPatchInstallerDelegate) instead, followed by IBuildInstaller::StartInstallation.")
diff --git a/Engine/Source/Runtime/Online/BuildPatchServices/Public/Interfaces/IBuildStatistics.h b/Engine/Source/Runtime/Online/BuildPatchServices/Public/Interfaces/IBuildStatistics.h
--- a/Engine/Source/Runtime/Online/BuildPatchServices/Public/Interfaces/IBuildStatistics.h
+++ b/Engine/Source/Runtime/Online/BuildPatchServices/Public/Interfaces/IBuildStatistics.h
@@ -128,25 +128,21 @@ namespace BuildPatchServices
 		/**
 		 * @return the size in chunks of the memory store.
 		 */
-		UE_DEPRECATED(5.6, "Appears to be unused, not replaced")
 		virtual int32 GetMemoryChunkStoreSize() const { return 0; }
 
 		/**
 		 * @return the number of chunks currently in the memory store.
 		 */
-		UE_DEPRECATED(5.6, "Appears to be unused, not replaced")
 		virtual int32 GetMemoryChunksInStore() const { return 0; }
 
 		/**
 		 * @return the number of chunks that have been booted from the memory store.
 		 */
-		UE_DEPRECATED(5.6, "Appears to be unused, not replaced")
 		virtual int32 GetMemoryChunksBooted() const { return 0; }
 
 		/**
 		 * @return the number of chunks currently in the memory store which are held due to multiple referencing.
 		 */
-		UE_DEPRECATED(5.6, "Appears to be unused, not replaced")
 		virtual int32 GetMemoryChunksRetained() const { return 0; }
 
 		/**
@@ -215,6 +211,21 @@ namespace BuildPatchServices
 		 */
 		virtual double GetDiskWriteByteSpeed() const = 0;
 
+		/**
+		 * @return the number of bytes downloaded so far.
+		 */
+		virtual int64 GetDownloadedBytes() const = 0;
+
+		/**
+		 * @return the number of bytes read so far.
+		 */
+		virtual int64 GetDiskReadBytes() const = 0;
+
+		/**
+		 * @return the number of bytes written so far.
+		 */
+		virtual int64 GetDiskWrittenBytes() const = 0;
+
 		/**
 		 * @return the total number of verify errors experienced during this installation.
 		 */
@@ -229,3 +240,65 @@ namespace BuildPatchServices
 	typedef TSharedPtr<IBuildStatistics> IBuildStatisticsPtr;
 	typedef TSharedRef<IBuildStatistics> IBuildStatisticsRef;
 }
+
+inline const TCHAR* LexToString(BuildPatchServices::EFileOperationState Value)
+{
+	using namespace BuildPatchServices;
+	static_assert(int32(EFileOperationState::NUM_States) == int32(EFileOperationState::VerifiedSuccess) + 1, "Please add support for the extra values below.");
+#define CASE_ENUM_TO_STR(_X_) case EFileOperationState::_X_: return TEXT(#_X_)
+	switch (Value)
+	{
+		CASE_ENUM_TO_STR(Unknown);
+		CASE_ENUM_TO_STR(PendingLocalChunkDbData);
+		CASE_ENUM_TO_STR(RetrievingLocalChunkDbData);
+		CASE_ENUM_TO_STR(PendingLocalInstallData);
+		CASE_ENUM_TO_STR(RetrievingLocalInstallData);
+		CASE_ENUM_TO_STR(PendingRemoteCloudData);
+		CASE_ENUM_TO_STR(RetrievingRemoteCloudData);
+		CASE_ENUM_TO_STR(PendingLocalDataStore);
+		CASE_ENUM_TO_STR(RetrievingLocalDataStore);
+		CASE_ENUM_TO_STR(DataInMemoryStore);
+		CASE_ENUM_TO_STR(Staged);
+		CASE_ENUM_TO_STR(Installed);
+		CASE_ENUM_TO_STR(Verifying);
+		CASE_ENUM_TO_STR(VerifiedFail);
+		CASE_ENUM_TO_STR(VerifiedSuccess);
+		CASE_ENUM_TO_STR(NUM_States);
+		default: return TEXT("InvalidOrMax");
+	}
+#undef CASE_ENUM_TO_STR
+}
+
+inline void LexFromString(BuildPatchServices::EFileOperationState& Value, const TCHAR* Buffer)
+{
+	using namespace BuildPatchServices;
+	static_assert(int32(EFileOperationState::NUM_States) == int32(EFileOperationState::VerifiedSuccess) + 1, "Please add support for the extra values below.");
+#define RETURN_IF_EQUAL(_X_) if (FCString::Stricmp(Buffer, TEXT(#_X_)) == 0) { Value = EFileOperationState::_X_; return; }
+	const TCHAR* const Prefix = TEXT("EFileOperationState::");
+	const SIZE_T PrefixLen = FCString::Strlen(Prefix);
+	if (FCString::Strnicmp(Buffer, Prefix, PrefixLen) == 0)
+	{
+		Buffer += PrefixLen;
+	}
+	RETURN_IF_EQUAL(Unknown);
+	RETURN_IF_EQUAL(PendingLocalChunkDbData);
+	RETURN_IF_EQUAL(RetrievingLocalChunkDbData);
+	RETURN_IF_EQUAL(PendingLocalInstallData);
+	RETURN_IF_EQUAL(RetrievingLocalInstallData);
+	RETURN_IF_EQUAL(PendingRemoteCloudData);
+	RETURN_IF_EQUAL(RetrievingRemoteCloudData);
+	RETURN_IF_EQUAL(PendingLocalDataStore);
+	RETURN_IF_EQUAL(RetrievingLocalDataStore);
+	RETURN_IF_EQUAL(DataInMemoryStore);
+	RETURN_IF_EQUAL(Staged);
+	RETURN_IF_EQUAL(Installed);
+	RETURN_IF_EQUAL(Verifying);
+	RETURN_IF_EQUAL(VerifiedFail);
+	RETURN_IF_EQUAL(VerifiedSuccess);
+	RETURN_IF_EQUAL(NUM_States);
+	// Did not match
+	Value = EFileOperationState::Unknown;
+#undef RETURN_IF_EQUAL
+}
+
+
diff --git a/Engine/Source/Runtime/Online/BuildPatchServices/Public/Interfaces/IPatchDataEnumeration.h b/Engine/Source/Runtime/Online/BuildPatchServices/Public/Interfaces/IPatchDataEnumeration.h
--- a/Engine/Source/Runtime/Online/BuildPatchServices/Public/Interfaces/IPatchDataEnumeration.h
+++ b/Engine/Source/Runtime/Online/BuildPatchServices/Public/Interfaces/IPatchDataEnumeration.h
@@ -2,7 +2,6 @@
 #pragma once
 
 #include "CoreMinimal.h"
-#include "Misc/Guid.h"
 
 namespace BuildPatchServices
 {
@@ -13,12 +12,14 @@ namespace BuildPatchServices
 
 		/**
 		 * Enumerates the data files used and outputs them to the OutputFile specified in the configuration
+		 * NOTE: THIS function is blocking and will not return until finished.
 		 * @return Success
 		 */
 		virtual bool Run() = 0;
 
 		/**
 		 * Enumerates the data files used and outputs them into the array provided
+		 * NOTE: THIS function is blocking and will not return until finished.
 		 * @param OutFiles		OUT Array of enumerated files
 		 * @return Success
 		 */
diff --git a/Engine/Source/Runtime/Online/HTTP/Private/Curl/CurlHttp.cpp b/Engine/Source/Runtime/Online/HTTP/Private/Curl/CurlHttp.cpp
--- a/Engine/Source/Runtime/Online/HTTP/Private/Curl/CurlHttp.cpp
+++ b/Engine/Source/Runtime/Online/HTTP/Private/Curl/CurlHttp.cpp
@@ -871,6 +871,7 @@ bool FCurlHttpRequest::StartThreadedRequest()
 
 		SetupOptionUnixSocketPath();
 		SetupOptionHttpVersion();
+		SetupOptionEnableCookies();
 
 		if (!FCurlHttpManager::CurlRequestOptions.LocalHostAddr.IsEmpty())
 		{
@@ -1048,6 +1049,16 @@ void FCurlHttpRequest::SetupOptionHttpVersion()
 	}
 }
 
+void FCurlHttpRequest::SetupOptionEnableCookies()
+{
+	const FString EnableCookies = GetOption(HttpRequestOptions::EnableCookies);
+	if (EnableCookies.ToBool())
+	{
+		// Enable the cookie engine without reading any initial cookies.
+		curl_easy_setopt(EasyHandle, CURLOPT_COOKIEFILE, "");
+	}
+}
+
 void FCurlHttpRequest::CleanupRequestHttpThread()
 {
 	curl_easy_setopt(EasyHandle, CURLOPT_SHARE, nullptr);
diff --git a/Engine/Source/Runtime/Online/HTTP/Private/Curl/CurlHttp.h b/Engine/Source/Runtime/Online/HTTP/Private/Curl/CurlHttp.h
--- a/Engine/Source/Runtime/Online/HTTP/Private/Curl/CurlHttp.h
+++ b/Engine/Source/Runtime/Online/HTTP/Private/Curl/CurlHttp.h
@@ -257,6 +257,7 @@ private:
 
 	void SetupOptionUnixSocketPath();
 	void SetupOptionHttpVersion();
+	void SetupOptionEnableCookies();
 
 	void LogCurlErrors();
 
diff --git a/Engine/Source/Runtime/Online/HTTP/Private/HttpRetrySystem.cpp b/Engine/Source/Runtime/Online/HTTP/Private/HttpRetrySystem.cpp
--- a/Engine/Source/Runtime/Online/HTTP/Private/HttpRetrySystem.cpp
+++ b/Engine/Source/Runtime/Online/HTTP/Private/HttpRetrySystem.cpp
@@ -193,6 +193,19 @@ void FHttpRetrySystem::FRequest::MoveToNextRetryDomain()
 	if (RetryDomains->ActiveIndex.CompareExchange(RetryDomainsIndex, NextDomainIndex))
 	{
 		RetryDomainsIndex = NextDomainIndex;
+
+		// We should updade config files only if two fields were explicitly set
+		if (!RetryDomains->ConfigSection.IsEmpty() && !RetryDomains->ConfigKeyName.IsEmpty())
+		{
+			GConfig->RemoveKey(*RetryDomains->ConfigSection, *RetryDomains->ConfigKeyName, GGameUserSettingsIni);
+
+			// If we come in a circle to the first index then no need to save it
+			if (RetryDomains->ActiveIndex > 0)
+			{
+				// Save the configuration to reuse latest available endpoint during the next Launcher session
+				GConfig->SetString(*RetryDomains->ConfigSection, *RetryDomains->ConfigKeyName, *RetryDomains->Domains[RetryDomains->ActiveIndex], GGameUserSettingsIni);
+			}
+		}
 	}
 	SetUrlFromRetryDomains();
 }
diff --git a/Engine/Source/Runtime/Online/HTTP/Private/Interfaces/IHttpRequest.cpp b/Engine/Source/Runtime/Online/HTTP/Private/Interfaces/IHttpRequest.cpp
--- a/Engine/Source/Runtime/Online/HTTP/Private/Interfaces/IHttpRequest.cpp
+++ b/Engine/Source/Runtime/Online/HTTP/Private/Interfaces/IHttpRequest.cpp
@@ -9,6 +9,7 @@ namespace HttpRequestOptions
 #if UE_HTTP_SUPPORT_UNIX_SOCKET
 	const FName UnixSocketPath("UnixSocketPath");
 #endif //UE_HTTP_SUPPORT_UNIX_SOCKET
+	const FName EnableCookies("EnableCookies");
 }
 
 const TCHAR* LexToString(EHttpRequestMode Mode)
diff --git a/Engine/Source/Runtime/Online/HTTP/Public/HttpRetrySystem.h b/Engine/Source/Runtime/Online/HTTP/Public/HttpRetrySystem.h
--- a/Engine/Source/Runtime/Online/HTTP/Public/HttpRetrySystem.h
+++ b/Engine/Source/Runtime/Online/HTTP/Public/HttpRetrySystem.h
@@ -50,6 +50,13 @@ namespace FHttpRetrySystem
 			, ActiveIndex(0)
 		{}
 
+		FRetryDomains(TArray<FString>&& InDomains, const int32 InActiveIndex, const FString& InConfigSection, const FString& InConfigKeyName)
+			: Domains(MoveTemp(InDomains))
+			, ActiveIndex(InActiveIndex)
+			, ConfigSection(InConfigSection)
+			, ConfigKeyName(InConfigKeyName)
+		{}
+
 		/** The domains to use */
 		TArray<FString> Domains;
 		/**
@@ -57,6 +64,12 @@ namespace FHttpRetrySystem
 		 * Domains are cycled through on some errors, and when we succeed on one domain, we remain on that domain until that domain results in an error
 		 */
 		TAtomic<int32> ActiveIndex;
+
+		/** Config section name that can be used for storing active domain for reusing during next sessions */
+		const FString ConfigSection;
+
+		/** Config key name that can be used for storing active domain for reusing during next sessions */
+		const FString ConfigKeyName;
 	};
 	typedef TSharedPtr<FRetryDomains, ESPMode::ThreadSafe> FRetryDomainsPtr;
 
@@ -83,6 +96,7 @@ namespace FHttpRetrySystem
 		HTTP_API bool IsValid() const;
 		HTTP_API float Compute(uint32 RetryNumber) const;
 	};
+	using FRetryExponentialBackoffCurvePtr = TSharedPtr<FExponentialBackoffCurve, ESPMode::ThreadSafe>;
 
 	/**
 	 * Read the number of seconds a HTTP request is throttled for from the response
diff --git a/Engine/Source/Runtime/Online/HTTP/Public/Interfaces/IHttpRequest.h b/Engine/Source/Runtime/Online/HTTP/Public/Interfaces/IHttpRequest.h
--- a/Engine/Source/Runtime/Online/HTTP/Public/Interfaces/IHttpRequest.h
+++ b/Engine/Source/Runtime/Online/HTTP/Public/Interfaces/IHttpRequest.h
@@ -148,6 +148,7 @@ namespace HttpRequestOptions
 #if UE_HTTP_SUPPORT_UNIX_SOCKET
 	extern HTTP_API const FName UnixSocketPath;
 #endif //UE_HTTP_SUPPORT_UNIX_SOCKET
+	extern HTTP_API const FName EnableCookies;
 }
 
 /**
