[Backout] - CL50979289

diff --git a/Engine/Source/Programs/BuildPatchTool/Private/ToolModes/PatchGenerationMode.cpp b/Engine/Source/Programs/BuildPatchTool/Private/ToolModes/PatchGenerationMode.cpp
new file mode 100644
--- /dev/null
+++ b/Engine/Source/Programs/BuildPatchTool/Private/ToolModes/PatchGenerationMode.cpp
@@ -0,0 +1,370 @@
+// Copyright Epic Games, Inc. All Rights Reserved.
+
+#include "ToolModes/PatchGenerationMode.h"
+#include "Misc/DefaultValueHelper.h"
+#include "Misc/Paths.h"
+#include "Misc/CommandLine.h"
+#include "Interfaces/IBuildPatchServicesModule.h"
+#include "BuildPatchFeatureLevel.h"
+#include "BuildPatchTool.h"
+
+using namespace BuildPatchTool;
+
+namespace Constants
+{
+	static const FString Comma(TEXT(","));
+	static const FString Equals(TEXT("="));
+	static const FString DoubleQuote(TEXT("\""));
+	static const FString SingleQuote(TEXT("'"));
+	static const FString Slash(TEXT("/"));
+	static const FString Backslash(TEXT("\\"));
+	static const FString Custom(TEXT("custom"));
+	static const FString CustomInt(TEXT("customint"));
+	static const FString CustomFloat(TEXT("customfloat"));
+}
+
+class FPatchGenerationToolMode : public IToolMode
+{
+public:
+	FPatchGenerationToolMode(IBuildPatchServicesModule& InBpsInterface)
+		: BpsInterface(InBpsInterface)
+		, ChunkWindowSize(1048576)
+		, bIgnoreOtherWindowSizes(false)
+	{}
+
+	virtual~FPatchGenerationToolMode()
+	{}
+
+	virtual EReturnCode Execute() override
+	{
+		// Parse commandline
+		if (ProcessCommandline() == false)
+		{
+			return EReturnCode::ArgumentProcessingError;
+		}
+
+		// Print help if requested
+		if (bHelp)
+		{
+			UE_LOG(LogBuildPatchTool, Display, TEXT("GENERATE PATCH DATA MODE"));
+			UE_LOG(LogBuildPatchTool, Display, TEXT("This tool supports generating chunk based patches. Chunk based patch data will be generated by default."));
+			UE_LOG(LogBuildPatchTool, Display, TEXT(""));
+			UE_LOG(LogBuildPatchTool, Display, TEXT("Required arguments:"));
+			UE_LOG(LogBuildPatchTool, Display, TEXT("  -mode=PatchGeneration    Must be specified to launch the tool in patch generation mode."));
+			UE_LOG(LogBuildPatchTool, Display, TEXT("  -FeatureLevel=Latest     Specifies the client feature level to output data for. See BuildPatchServices::EFeatureLevel for possible values."));
+			UE_LOG(LogBuildPatchTool, Display, TEXT("  -BuildRoot=\"\"            Specifies in quotes the directory containing the build image to be read."));
+			UE_LOG(LogBuildPatchTool, Display, TEXT("  -CloudDir=\"\"             Specifies in quotes the cloud directory where existing data will be recognized from, and new data added to."));
+			UE_LOG(LogBuildPatchTool, Display, TEXT("  -AppName=\"\"              Specifies in quotes, the name of the app."));
+			UE_LOG(LogBuildPatchTool, Display, TEXT("  -BuildVersion=\"\"         Specifies in quotes, the version string for the build image."));
+			UE_LOG(LogBuildPatchTool, Display, TEXT("  -AppLaunch=\"\"            Specifies in quotes, the path to the app executable, must be relative to, and inside of BuildRoot."));
+			UE_LOG(LogBuildPatchTool, Display, TEXT("  -AppArgs=\"\"              Specifies in quotes, the commandline to send to the app on launch."));
+			UE_LOG(LogBuildPatchTool, Display, TEXT(""));
+			UE_LOG(LogBuildPatchTool, Display, TEXT("Optional arguments:"));
+			UE_LOG(LogBuildPatchTool, Display, TEXT("  -AppID=123456                Specifies without quotes, the ID number for the app. This will default to 0 if not provided."));
+			UE_LOG(LogBuildPatchTool, Display, TEXT("  -FileList=\"\"                 Specifies in quotes, the path to a text file containing BuildRoot relative files to be included in the build."));
+			UE_LOG(LogBuildPatchTool, Display, TEXT("  -FileIgnoreList=\"\"           Specifies in quotes, the path to a text file containing BuildRoot relative files, separated by \\r\\n line endings, to not be included in the build."));
+			UE_LOG(LogBuildPatchTool, Display, TEXT("  -FileAttributeList=\"\"        Specifies in quotes, the path to a text file containing quoted BuildRoot relative files followed by optional attribute keywords readonly compressed executable, separated by \\r\\n line endings. These attribute will be applied when build is installed client side."));
+			UE_LOG(LogBuildPatchTool, Display, TEXT("  -PrereqIds=\"\"                Specifies in quotes, a comma-separated list of identifiers that the prerequisites satisfy. At install time, a machine which already has installed prerequisites with all of these ids will skip prerequisite installation."));
+			UE_LOG(LogBuildPatchTool, Display, TEXT("  -PrereqName=\"\"               Specifies in quotes, the display name xfor the prerequisites installer."));
+			UE_LOG(LogBuildPatchTool, Display, TEXT("  -PrereqPath=\"\"               Specifies in quotes, the prerequisites installer to launch on successful product install."));
+			UE_LOG(LogBuildPatchTool, Display, TEXT("                               This path supports a string replace for \"$[RootDirectory]\". This will be replaced with the root path before executing. The replacement will include trailing /."));
+			UE_LOG(LogBuildPatchTool, Display, TEXT("  -PrereqArgs=\"\"               Specifies in quotes, the commandline to send to prerequisites installer on launch."));
+			UE_LOG(LogBuildPatchTool, Display, TEXT("                               This value supports string replacements for \"$[RootDirectory]\" and also \"$[LogDirectory]\". LogDirectory is the path to the program's log output directory so your prereq could create logs there. The replacement will include trailing /."));
+			UE_LOG(LogBuildPatchTool, Display, TEXT("                               \"$[Quote]\" can also be used to get a quote character, this is important because the BPT commandline already uses quotes for token parsing."));
+			UE_LOG(LogBuildPatchTool, Display, TEXT("  -DataAgeThreshold=12.5       Specified the maximum age (in days) of existing manifest files whose referenced patch data can be reused in the generated manifest."));
+			UE_LOG(LogBuildPatchTool, Display, TEXT("  -Custom=\"field=value\"        Adds a custom string field to the build manifest."));
+			UE_LOG(LogBuildPatchTool, Display, TEXT("  -CustomInt=\"field=number\"    Adds a custom int64 field to the build manifest."));
+			UE_LOG(LogBuildPatchTool, Display, TEXT("  -CustomFloat=\"field=number\"  Adds a custom double field to the build manifest."));
+			UE_LOG(LogBuildPatchTool, Display, TEXT("  -OutputFilename=\"\"           Specifies in quotes an override for the output manifest filename. Extension of .manifest will be added if not present."));
+			UE_LOG(LogBuildPatchTool, Display, TEXT("  -ChunkWindowSize=1000000     Specifies in bytes, the data window size that should be used when saving new chunks. Default is 1048576 (1MiB)."));
+			UE_LOG(LogBuildPatchTool, Display, TEXT("  -IgnoreOtherWindowSizes      If provided, the generation code will only accept chunk matches that are the same as ChunkWindowSize."));
+			UE_LOG(LogBuildPatchTool, Display, TEXT(""));
+			UE_LOG(LogBuildPatchTool, Display, TEXT("NB: If -DataAgeThreshold is not supplied on the command-line, then all existing data is eligible for reuse in the generated manifest."));
+			UE_LOG(LogBuildPatchTool, Display, TEXT("NB: If -OutputFilename is not supplied on the command-line, the default of AppNameBuildVersion.manifest will be used."));
+			UE_LOG(LogBuildPatchTool, Display, TEXT("NB:    -OutputFilename must be a clean filename with no path."));
+			UE_LOG(LogBuildPatchTool, Display, TEXT("NB: If -ChunkWindowSize provided, it is clamped max 10485760 (10MiB) to min 32000 (32KB)."));
+			UE_LOG(LogBuildPatchTool, Display, TEXT(""));
+			return EReturnCode::OK;
+		}
+
+		// Check existence of the file list
+		if (!FileList.IsEmpty() && !FPaths::FileExists(FileList))
+		{
+			UE_LOG(LogBuildPatchTool, Error, TEXT("Provided file list was not found %s"), *FileList);
+			return EReturnCode::FileNotFound;
+		}
+
+		// Check existence of file ignore list
+		if (!FileIgnoreList.IsEmpty() && !FPaths::FileExists(FileIgnoreList))
+		{
+			UE_LOG(LogBuildPatchTool, Error, TEXT("Provided file ignore list was not found %s"), *FileIgnoreList);
+			return EReturnCode::FileNotFound;
+		}
+
+		// Check existence of file attributes list
+		if (!FileAttributeList.IsEmpty() && !FPaths::FileExists(FileAttributeList))
+		{
+			UE_LOG(LogBuildPatchTool, Error, TEXT("Provided file attribute list was not found %s"), *FileAttributeList);
+			return EReturnCode::FileNotFound;
+		}
+
+		// Default the OutputFilename if not provided
+		if (OutputFilename.IsEmpty())
+		{
+			OutputFilename = FDefaultValueHelper::RemoveWhitespaces(AppName + BuildVersion) + TEXT(".manifest");
+		}
+		// Otherwise check the parameter
+		else
+		{
+			if (OutputFilename.Contains(TEXT("/")))
+			{
+				UE_LOG(LogBuildPatchTool, Error, TEXT("Provided OutputFilename should be clean filename only. Invalid arg: %s"), *OutputFilename);
+				return EReturnCode::ArgumentProcessingError;
+			}
+		}
+
+		// Setup and run
+		BuildPatchServices::FChunkBuildConfiguration Settings;
+		const uint32 DefaultChunkWindowSize = Settings.OutputChunkWindowSize;
+		if (!BuildPatchServices::FeatureLevelFromString(*FeatureLevel, Settings.FeatureLevel))
+		{
+			UE_LOG(LogBuildPatchTool, Error, TEXT("Provided FeatureLevel is not recognised. Invalid arg: -FeatureLevel=%s"), *FeatureLevel);
+			return EReturnCode::ArgumentProcessingError;
+		}
+		Settings.RootDirectory = BuildRoot;
+		Settings.AppId = TCString<TCHAR>::Atoi64(*AppId);
+		Settings.AppName = AppName;
+		Settings.BuildVersion = BuildVersion;
+		Settings.LaunchExe = AppLaunch;
+		Settings.LaunchCommand = AppArgs;
+		Settings.InputListFile = FileList;
+		Settings.IgnoreListFile = FileIgnoreList;
+		Settings.AttributeListFile = FileAttributeList;
+		Settings.PrereqIds = PrereqIdsSet;
+		Settings.PrereqName = PrereqName;
+		Settings.PrereqPath = PrereqPath;
+		Settings.PrereqArgs = PrereqArgs;
+		Settings.DataAgeThreshold = TCString<TCHAR>::Atod(*DataAgeThreshold);
+		Settings.bShouldHonorReuseThreshold = DataAgeThreshold.IsEmpty() == false;
+		Settings.OutputChunkWindowSize = ChunkWindowSize;
+		Settings.bShouldMatchAnyWindowSize = !bIgnoreOtherWindowSizes && Settings.FeatureLevel >= BuildPatchServices::EFeatureLevel::VariableSizeChunks;
+		Settings.CustomFields = CustomFields;
+		Settings.CloudDirectory = CloudDir;
+		Settings.OutputFilename = OutputFilename;
+
+		// Check feature compatibility
+		const bool bHasCustomFields = Settings.CustomFields.Num() > 0;
+		if (Settings.FeatureLevel < BuildPatchServices::EFeatureLevel::CustomFields && bHasCustomFields)
+		{
+			UE_LOG(LogBuildPatchTool, Error, TEXT("Invalid args: FeatureLevel %s is not compatible with Custom, CustomInt, or CustomFloat."), BuildPatchServices::FeatureLevelToString(Settings.FeatureLevel));
+			return EReturnCode::ArgumentProcessingError;
+		}
+		const bool bHasAnyPrereqInfo = !PrereqName.IsEmpty() || !PrereqPath.IsEmpty() || !PrereqArgs.IsEmpty();
+		if (Settings.FeatureLevel < BuildPatchServices::EFeatureLevel::StoresPrerequisitesInfo && bHasAnyPrereqInfo)
+		{
+			UE_LOG(LogBuildPatchTool, Error, TEXT("Invalid args: FeatureLevel %s is not compatible with PrereqName, PrereqPath, or PrereqArgs."), BuildPatchServices::FeatureLevelToString(Settings.FeatureLevel));
+			return EReturnCode::ArgumentProcessingError;
+		}
+		const bool bHasPrereqIds = Settings.PrereqIds.Num() > 0;
+		if (Settings.FeatureLevel < BuildPatchServices::EFeatureLevel::StoresPrerequisiteIds && bHasPrereqIds)
+		{
+			UE_LOG(LogBuildPatchTool, Error, TEXT("Invalid args: FeatureLevel %s is not compatible with PrereqIds."), BuildPatchServices::FeatureLevelToString(Settings.FeatureLevel));
+			return EReturnCode::ArgumentProcessingError;
+		}
+		const bool bHasNonDefaultWindowSize = Settings.OutputChunkWindowSize != DefaultChunkWindowSize;
+		if (Settings.FeatureLevel < BuildPatchServices::EFeatureLevel::VariableSizeChunks && bHasNonDefaultWindowSize)
+		{
+			UE_LOG(LogBuildPatchTool, Error, TEXT("Invalid args: FeatureLevel %s is not compatible with -ChunkWindowSize=%u."), BuildPatchServices::FeatureLevelToString(Settings.FeatureLevel), Settings.OutputChunkWindowSize);
+			return EReturnCode::ArgumentProcessingError;
+		}
+
+		// Run the build generation
+		bool bSuccess = BpsInterface.ChunkBuildDirectory(Settings);
+		return bSuccess ? EReturnCode::OK : EReturnCode::ToolFailure;
+	}
+
+private:
+
+	bool ParsePrereqIds(const FString& ParamValue, TSet<FString>& OutPrereqIds)
+	{
+		if (ParamValue.Contains(Constants::Slash)
+			|| ParamValue.Contains(Constants::Backslash)
+			|| ParamValue.Contains(Constants::DoubleQuote)
+			|| ParamValue.Contains(Constants::Slash))
+		{
+			return false;
+		}
+
+		TArray<FString> ParamValues;
+		ParamValue.ParseIntoArray(ParamValues, *Constants::Comma);
+		OutPrereqIds.Append(ParamValues);
+		return true;
+	}
+
+	bool ParseCustomField(const FString& Switch, TMap<FString, FVariant>& Fields)
+	{
+		FString Type, Left, Right;
+
+		Switch.Split(Constants::Equals, &Type, &Right);
+		Type.ToLowerInline();
+		Right.Split(Constants::Equals, &Left, &Right);
+		Left.TrimStartAndEndInline();
+		Right.TrimStartAndEndInline();
+		if (Type.Equals(Constants::Custom, ESearchCase::CaseSensitive))
+		{
+			CustomFields.Add(Left, FVariant(Right));
+		}
+		else if (Type.Equals(Constants::CustomInt, ESearchCase::CaseSensitive))
+		{
+			if (!Right.IsNumeric())
+			{
+				UE_LOG(LogBuildPatchTool, Error, TEXT("An error occurred processing numeric token from commandline -%s"), *Switch);
+				return false;
+			}
+			CustomFields.Add(Left, FVariant(TCString<TCHAR>::Atoi64(*Right)));
+		}
+		else if (Type.Equals(Constants::CustomFloat, ESearchCase::CaseSensitive))
+		{
+			if (!Right.IsNumeric())
+			{
+				UE_LOG(LogBuildPatchTool, Error, TEXT("An error occurred processing numeric token from commandline -%s"), *Switch);
+				return false;
+			}
+			CustomFields.Add(Left, FVariant(TCString<TCHAR>::Atod(*Right)));
+		}
+
+		return true;
+	}
+
+	bool ProcessCommandline()
+	{
+#define PARSE_SWITCH(Switch) ParseSwitch(TEXT(#Switch "="), Switch, Switches)
+		TArray<FString> Tokens, Switches;
+		FCommandLine::Parse(FCommandLine::Get(), Tokens, Switches);
+
+		bHelp = ParseOption(TEXT("help"), Switches);
+		if (bHelp)
+		{
+			return true;
+		}
+
+		// Grab the FeatureLevel. This is required param but safe to default, we can change this to a warning after first release, and then an error later, as part of a friendly roll out.
+		PARSE_SWITCH(FeatureLevel);
+		FeatureLevel.TrimStartAndEndInline();
+		if (FeatureLevel.IsEmpty())
+		{
+			UE_LOG(LogBuildPatchTool, Log, TEXT("FeatureLevel was not provided, defaulting to LatestJson. Please provide the FeatureLevel commandline argument which matches the existing client support."));
+			FeatureLevel = TEXT("LatestJson");
+		}
+
+		// Get all required parameters
+		if (!(PARSE_SWITCH(CloudDir)
+		   && PARSE_SWITCH(BuildRoot)
+		   && PARSE_SWITCH(AppName)
+		   && PARSE_SWITCH(BuildVersion)
+		   && PARSE_SWITCH(AppLaunch)
+		   && PARSE_SWITCH(AppArgs)))
+		{
+			UE_LOG(LogBuildPatchTool, Error, TEXT("CloudDir, BuildRoot, AppName, BuildVersion, AppLaunch, and AppArgs are required parameters"));
+			return false;
+		}
+		NormalizeUriPath(CloudDir);
+		NormalizeUriPath(BuildRoot);
+		NormalizeUriFile(AppLaunch);
+
+		// Get optional parameters
+		PARSE_SWITCH(AppId);
+		PARSE_SWITCH(FileList);
+		PARSE_SWITCH(FileIgnoreList);
+		PARSE_SWITCH(FileAttributeList);
+		PARSE_SWITCH(PrereqIds);
+		PARSE_SWITCH(PrereqName);
+		PARSE_SWITCH(PrereqPath);
+		PARSE_SWITCH(PrereqArgs);
+		PARSE_SWITCH(DataAgeThreshold);
+		PARSE_SWITCH(ChunkWindowSize);
+		bIgnoreOtherWindowSizes = ParseOption(TEXT("IgnoreOtherWindowSizes"), Switches);
+		PARSE_SWITCH(OutputFilename);
+		NormalizeUriFile(FileList);
+		NormalizeUriFile(FileIgnoreList);
+		NormalizeUriFile(FileAttributeList);
+		NormalizeUriFile(PrereqPath);
+		NormalizeUriFile(OutputFilename);
+
+		// Check manifest file extension.
+		const TCHAR* ManifestExtension = TEXT(".manifest");
+		if (!OutputFilename.IsEmpty() && !OutputFilename.EndsWith(ManifestExtension))
+		{
+			OutputFilename += ManifestExtension;
+		}
+
+		// Clamp ChunkWindowSize to sane range.
+		const uint32 RequestedChunkWindowSize = ChunkWindowSize;
+		ChunkWindowSize = FMath::Clamp<uint32>(ChunkWindowSize, 32000, 10485760);
+		if (RequestedChunkWindowSize != ChunkWindowSize)
+		{
+			UE_LOG(LogBuildPatchTool, Warning, TEXT("Requested -ChunkWindowSize=%u is outside of allowed range 10485760 >= n >= 32000. Please update your args to be within range. Continuing with %u."), RequestedChunkWindowSize, ChunkWindowSize);
+		}
+
+		// Check numeric values
+		if (!AppId.IsEmpty() && !AppId.IsNumeric())
+		{
+			UE_LOG(LogBuildPatchTool, Error, TEXT("An error occurred processing numeric token from commandline -AppId=%s"), *AppId);
+			return false;
+		}
+		if (!DataAgeThreshold.IsEmpty() && !DataAgeThreshold.IsNumeric())
+		{
+			UE_LOG(LogBuildPatchTool, Error, TEXT("An error occurred processing numeric token from commandline -DataAgeThreshold=%s"), *DataAgeThreshold);
+			return false;
+		}
+
+		// Get custom fields to add to manifest
+		// These are optional, but a failure to parse one is an error
+		for (const FString& Switch : Switches)
+		{
+			if (Switch.StartsWith(Constants::Custom) && !ParseCustomField(Switch, CustomFields))
+			{
+				return false;
+			}
+		}
+
+		if (!PrereqIds.IsEmpty() && !ParsePrereqIds(PrereqIds, PrereqIdsSet))
+		{
+			UE_LOG(LogBuildPatchTool, Error, TEXT("An error occurred processing comma-separated list from commandline -PrereqIds=%s"), *PrereqIds);
+			return false;
+		}
+
+		return true;
+#undef PARSE_SWITCH
+	}
+
+private:
+	IBuildPatchServicesModule& BpsInterface;
+	bool bHelp;
+	FString FeatureLevel;
+	FString BuildRoot;
+	FString CloudDir;
+	FString AppId;
+	FString AppName;
+	FString BuildVersion;
+	FString AppLaunch;
+	FString AppArgs;
+	FString PrereqIds;
+	TSet<FString> PrereqIdsSet;
+	FString PrereqName;
+	FString PrereqPath;
+	FString PrereqArgs;
+	FString FileList;
+	FString FileIgnoreList;
+	FString FileAttributeList;
+	FString DataAgeThreshold;
+	uint32 ChunkWindowSize;
+	bool bIgnoreOtherWindowSizes;
+	TMap<FString, FVariant> CustomFields;
+	FString OutputFilename;
+};
+
+BuildPatchTool::IToolModeRef BuildPatchTool::FPatchGenerationToolModeFactory::Create(IBuildPatchServicesModule& BpsInterface)
+{
+	return MakeShareable(new FPatchGenerationToolMode(BpsInterface));
+}
diff --git a/Engine/Source/Programs/BuildPatchTool/Private/ToolModes/VerifyChunksMode.cpp b/Engine/Source/Programs/BuildPatchTool/Private/ToolModes/VerifyChunksMode.cpp
--- a/Engine/Source/Programs/BuildPatchTool/Private/ToolModes/VerifyChunksMode.cpp
+++ b/Engine/Source/Programs/BuildPatchTool/Private/ToolModes/VerifyChunksMode.cpp
@@ -1,22 +1,19 @@
 // Copyright Epic Games, Inc. All Rights Reserved.
 
 #include "ToolModes/VerifyChunksMode.h"
-
-#include "Algo/Find.h"
-#include "BuildPatchTool.h"
 #include "Interfaces/IBuildPatchServicesModule.h"
+#include "BuildPatchTool.h"
 #include "Misc/CommandLine.h"
 #include "Misc/Paths.h"
-#include "ToolModes/ToolModesHelp.h"
+#include "Algo/Find.h"
 
 using namespace BuildPatchTool;
 
 class FVerifyChunksToolMode : public IToolMode
 {
 public:
-	FVerifyChunksToolMode(IBuildPatchServicesModule& InBpsInterface, const TCHAR* InCommandLine)
+	FVerifyChunksToolMode(IBuildPatchServicesModule& InBpsInterface)
 		: BpsInterface(InBpsInterface)
-		, CommandLine(InCommandLine)
 	{}
 
 	virtual ~FVerifyChunksToolMode()
@@ -33,18 +30,20 @@ public:
 		// Print help if requested.
 		if (bHelp)
 		{
-			PrintHelp<FVerifyChunksToolModeHelp>();
+			UE_LOG(LogBuildPatchTool, Display, TEXT("VERIFY CHUNKS MODE"));
+			UE_LOG(LogBuildPatchTool, Display, TEXT("This tool mode allows you to verify the integrity of patch data. It will load chunk or chunkdb files to verify they are not corrupt."));
+			UE_LOG(LogBuildPatchTool, Display, TEXT(""));
+			UE_LOG(LogBuildPatchTool, Display, TEXT("Required arguments:"));
+			UE_LOG(LogBuildPatchTool, Display, TEXT("  -mode=VerifyChunks  Must be specified to launch the tool in verify chunks mode."));
+			UE_LOG(LogBuildPatchTool, Display, TEXT("  -SearchPath=\"\"      Specifies in quotes the directory path which contains data to verify."));
+			UE_LOG(LogBuildPatchTool, Display, TEXT(""));
+			UE_LOG(LogBuildPatchTool, Display, TEXT("Optional arguments:"));
+			UE_LOG(LogBuildPatchTool, Display, TEXT("  -OutputFile=\"\"      When specified, full file path for each bad data will be saved to this file as \\r\\n seperated list."));
+			UE_LOG(LogBuildPatchTool, Display, TEXT(""));
+			UE_LOG(LogBuildPatchTool, Display, TEXT("NB: All checks are logged, normal log for good data, error log for any bad data found."));
 			return EReturnCode::OK;
 		}
 
-		LogConfig();
-
-		// Register secrets.
-		for (const TPair<FGuid, TArray<uint8>>& IdKeyPair : EncryptionSecrets)
-		{
-			BpsInterface.RegisterEncryptionSecret(IdKeyPair.Key, IdKeyPair.Value);
-		}
-
 		// Run the enumeration routine.
 		bool bSuccess = BpsInterface.VerifyChunkData(SearchPath, OutputFile);
 		return bSuccess ? EReturnCode::OK : EReturnCode::ToolFailure;
@@ -57,7 +56,7 @@ private:
 #define HAS_SWITCH(Switch) (Algo::FindByPredicate(Switches, [](const FString& Elem){ return Elem.StartsWith(TEXT(#Switch "="));}) != nullptr)
 #define PARSE_SWITCH(Switch) ParseSwitch(TEXT(#Switch "="), Switch, Switches)
 		TArray<FString> Tokens, Switches;
-		FCommandLine::Parse(CommandLine, Tokens, Switches);
+		FCommandLine::Parse(FCommandLine::Get(), Tokens, Switches);
 
 		bHelp = ParseOption(TEXT("help"), Switches);
 		if (bHelp)
@@ -71,51 +70,25 @@ private:
 			UE_LOG(LogBuildPatchTool, Error, TEXT("SearchPath is a required parameter"));
 			return false;
 		}
-
-		if (SearchPath.IsEmpty())
-		{
-			UE_LOG(LogBuildPatchTool, Error, TEXT("SearchPath argument can't be empty"));
-			return false;
-		}
-
 		NormalizeUriPath(SearchPath);
 
 		// Get optional parameters.
 		PARSE_SWITCH(OutputFile);
 		NormalizeUriFile(OutputFile);
 
-		// Parse all the secret key and ID pairs.
-		FString ParseErrorMessage;
-		if (!ParsePairs(TEXT("EncryptionSecretId="), TEXT("EncryptionSecretKey="), Switches, EncryptionSecrets, ParseErrorMessage) && !ParseErrorMessage.IsEmpty())
-		{
-			UE_LOG(LogBuildPatchTool, Error, TEXT("%s"), *ParseErrorMessage);
-			return false;
-		}
-
 		return true;
 #undef PARSE_SWITCH
 #undef HAS_SWITCH
 	}
 
-	void LogConfig() const
-	{
-		UE_LOG(LogBuildPatchTool, Log, TEXT("-----Configuration for VerifyChunksMode------"));
-		UE_LOG(LogBuildPatchTool, Log, TEXT("   SearchPath: %s"), *SearchPath);
-		UE_LOG(LogBuildPatchTool, Log, TEXT("   OutputFile: %s"), *OutputFile);
-		for (const TPair<FGuid, TArray<uint8>>& EncryptionSecret : EncryptionSecrets)
-		{
-			UE_LOG(LogBuildPatchTool, Log, TEXT("   AvailableSecretId: %s"), *EncryptionSecret.Key.ToString());
-		}
-	}
-
 private:
-	static TCHAR const* const MODE_NAME;
 	IBuildPatchServicesModule& BpsInterface;
-	const TCHAR* CommandLine;
 	bool bHelp;
 	FString SearchPath;
 	FString OutputFile;
-	TMap<FGuid, TArray<uint8>> EncryptionSecrets;
 };
 
-IMPLEMENT_BPT_MODE(VerifyChunks, FVerifyChunksToolMode);
+BuildPatchTool::IToolModeRef BuildPatchTool::FVerifyChunksToolModeFactory::Create(IBuildPatchServicesModule& BpsInterface)
+{
+	return MakeShareable(new FVerifyChunksToolMode(BpsInterface));
+}
diff --git a/Engine/Source/Programs/BuildPatchTool/Public/BuildPatchTool.h b/Engine/Source/Programs/BuildPatchTool/Public/BuildPatchTool.h
--- a/Engine/Source/Programs/BuildPatchTool/Public/BuildPatchTool.h
+++ b/Engine/Source/Programs/BuildPatchTool/Public/BuildPatchTool.h
@@ -15,47 +15,7 @@ namespace BuildPatchTool
 		UnknownToolMode,
 		FileNotFound,
 		ToolFailure,
-		OnlineMcpError,
-		DeprecatedVersion,
-		DirectoryNotFound,
-		MissingCredentials,
-		InvalidCredentials,
-		InvalidData,
-
-		// !! Always after the last code, signifies the last return code plus 1 to allow the following Last alias.
-		LastPlusOne,
-		// An alias for the actual last return code.
-		Last = (LastPlusOne - 1),
 
 		Crash = 255
 	};
-}
-
-namespace CommandLineHelpers
-{
-	bool ParseSwitch(const TCHAR* InSwitch, FString& Value, const TArray<FString>& Switches);
-}
-
-inline const TCHAR* LexToString(BuildPatchTool::EReturnCode ReturnCode)
-{
-	static_assert(BuildPatchTool::EReturnCode::Last == BuildPatchTool::EReturnCode::InvalidData, "Please add support for the extra values below.");
-#define CASE_ENUM_TO_STR(Value) case BuildPatchTool::EReturnCode::Value: return TEXT(#Value)
-	switch (ReturnCode)
-	{
-		CASE_ENUM_TO_STR(OK);
-		CASE_ENUM_TO_STR(UnknownError);
-		CASE_ENUM_TO_STR(ArgumentProcessingError);
-		CASE_ENUM_TO_STR(UnknownToolMode);
-		CASE_ENUM_TO_STR(FileNotFound);
-		CASE_ENUM_TO_STR(ToolFailure);
-		CASE_ENUM_TO_STR(OnlineMcpError);
-		CASE_ENUM_TO_STR(DeprecatedVersion);
-		CASE_ENUM_TO_STR(DirectoryNotFound);
-		CASE_ENUM_TO_STR(MissingCredentials);
-		CASE_ENUM_TO_STR(InvalidCredentials);
-		CASE_ENUM_TO_STR(InvalidData);
-		CASE_ENUM_TO_STR(Crash);
-	default: return TEXT("InvalidOrMax");
-	}
-#undef CASE_ENUM_TO_STR
-}
+}
\ No newline at end of file
diff --git a/Engine/Source/Programs/BuildPatchTool/Public/Interfaces/ToolMode.h b/Engine/Source/Programs/BuildPatchTool/Public/Interfaces/ToolMode.h
--- a/Engine/Source/Programs/BuildPatchTool/Public/Interfaces/ToolMode.h
+++ b/Engine/Source/Programs/BuildPatchTool/Public/Interfaces/ToolMode.h
@@ -3,64 +3,6 @@
 
 #include "CoreMinimal.h"
 #include "Interfaces/IBuildPatchServicesModule.h"
-#include "Algo/AllOf.h"
-#include "HAL/FileManager.h"
-#include "BuildPatchTool.h"
-
-#define IMPLEMENT_BPT_MODE(ModeName, ModeClass) \
-class ModeClass; \
-TCHAR const * const ModeClass::MODE_NAME = TEXT(#ModeName); \
-struct FAutoRegister##ModeName \
-{ \
-	FAutoRegister##ModeName() \
-	{ \
-		using namespace BuildPatchTool; \
-		FToolModeFactory::GetModeRegistryDelegate().AddLambda([](TMap<FString, FToolModeFactory::FToolModeCreateFunc>& OutModes) \
-			{ \
-				OutModes.Add(TEXT(#ModeName), [](IBuildPatchServicesModule& BpsInterface, const TCHAR* CommandLine){ return new ModeClass(BpsInterface, CommandLine); }); \
-			}); \
-	} \
-} AutoRegister##ModeName;
-
-#define IMPLEMENT_BPT_MODE_ALIAS(ModeName, ModeClass) \
-struct FAutoRegister##ModeName \
-{ \
-	FAutoRegister##ModeName() \
-	{ \
-		using namespace BuildPatchTool; \
-		FToolModeFactory::GetModeRegistryDelegate().AddLambda([](TMap<FString, FToolModeFactory::FToolModeCreateFunc>& OutModes) \
-			{ \
-				OutModes.Add(TEXT(#ModeName), [](IBuildPatchServicesModule& BpsInterface, const TCHAR* CommandLine){ return new ModeClass(BpsInterface, CommandLine); }); \
-			}); \
-	} \
-} AutoRegister##ModeName;
-
-namespace Constants
-{
-	static const FString Comma(TEXT(","));
-	static const FString Equals(TEXT("="));
-	static const FString DoubleQuote(TEXT("\""));
-	static const FString SingleQuote(TEXT("'"));
-	static const FString Slash(TEXT("/"));
-	static const FString Backslash(TEXT("\\"));
-	static const FString Custom(TEXT("custom"));
-	static const FString CustomInt(TEXT("customint"));
-	static const FString CustomFloat(TEXT("customfloat"));
-}
-
-// Detect if a template parameter is a TSet
-template<typename T>
-struct is_ue_set : std::false_type {};
-
-template<typename... T>
-struct is_ue_set<TSet<T...>> : std::true_type {};
-
-// Detect if template parameter is a TArray
-template<typename T>
-struct is_ue_array : std::false_type {};
-
-template<typename... T>
-struct is_ue_array<TArray<T...>> : std::true_type {};
 
 namespace BuildPatchTool
 {
@@ -73,20 +15,8 @@ namespace BuildPatchTool
 		TCHAR const * const QuoteStr = TEXT("\"");
 
 	public:
-		virtual ~IToolMode() { }
-
 		virtual EReturnCode Execute() = 0;
 
-		bool bHelp = false;
-		template<typename T>
-		void PrintHelp()
-		{
-			for (int32 i = 0; i < T::NumLines; i++)
-			{
-				UE_LOG(LogBuildPatchTool, Display, TEXT("%s"), T::Text[i]);
-			}
-		}
-
 		/**
 		 * Helper for normalizing a URI referring to a file. We make sure not to cause issues with URIs that separate a protocol and authority with ://
 		 * and windows UNC paths that begin with \\. Windows UNC paths still work in UE4 code with //.
@@ -99,14 +29,6 @@ namespace BuildPatchTool
 			UriFile.ReplaceInline(TEXT("\\"), TEXT("/"), ESearchCase::CaseSensitive);
 		}
 
-		void NormalizeUriFileInKeys(TArray<TPair<FString, FString>>& Pairs)
-		{
-			for (TPair<FString, FString>& Pair : Pairs)
-			{
-				NormalizeUriFile(Pair.Key);
-			}
-		}
-
 		/**
 		 * Helper for normalizing a URI referring to a path. We make sure not to cause issues with URIs that separate a protocol and authority with ://
 		 * and windows UNC paths that begin with \\. Windows UNC paths still work in UE4 code with //.
@@ -127,21 +49,6 @@ namespace BuildPatchTool
 			}
 		}
 
-		bool IsEmptyDirectory(const FString& Directory)
-		{
-			IFileManager& FM = IFileManager::Get();
-			TArray<FString> FoundFiles;
-			FM.FindFilesRecursive(FoundFiles, *Directory, TEXT("*"), true, false);
-			for (const FString& File : FoundFiles)
-			{
-				if (FM.FileSize(*File) > 0)
-				{
-					return false;
-				}
-			}
-			return true;
-		}
-
 		/**
 		 * Helper for parsing a switch from an array of switches, usually produced using FCommandLine::Parse(..)
 		 *
@@ -157,11 +64,14 @@ namespace BuildPatchTool
 			checkSlow(InSwitch != nullptr);
 			checkSlow(InSwitch[FCString::Strlen(InSwitch)-1] == TEXT('='));
 
-			if (const FString* Switch = Switches.FindByPredicate([InSwitch](const FString& Switch) { return Switch.StartsWith(InSwitch); }))
+			for (const FString& Switch : Switches)
 			{
-				FString StringValue;
-				Switch->Split(EqualsStr, nullptr, &StringValue);
-				return ParseValue(StringValue, Value);
+				if (Switch.StartsWith(InSwitch))
+				{
+					FString StringValue;
+					Switch.Split(EqualsStr, nullptr, &StringValue);
+					return ParseValue(StringValue, Value);
+				}
 			}
 			return false;
 		}
@@ -204,77 +114,6 @@ namespace BuildPatchTool
 			return Switches.Contains(InSwitch);
 		}
 
-		template<typename TFirst, typename TSecond, typename TContainer>
-		bool ParsePairs(const TCHAR* InFirstSwitch, const TCHAR* InSecondSwitch, const TArray<FString>& Switches, TContainer& OutPairs, FString& ErrorMessage)
-		{
-			bool bFoundAnyPair = false;
-			bool bFoundFirst = false;
-			bool bFoundSecond = false;
-			TFirst FirstValue;
-			TSecond SecondValue;
-			for (int32 SwitchIndex = 0; SwitchIndex < Switches.Num(); SwitchIndex++)
-			{
-				const FString& InputSwitch = Switches[SwitchIndex];
-				const bool bIsFirst = InputSwitch.StartsWith(InFirstSwitch);
-				const bool bIsSecond = InputSwitch.StartsWith(InSecondSwitch);
-				if (bIsFirst || bIsSecond)
-				{
-					 // Check we are alternating
-					if ((bIsFirst && bFoundFirst)
-					 || (bIsSecond && bFoundSecond))
-					{
-						// The error will be set in the check after this loop.
-						break;
-					}
-					FString ValStr;
-					InputSwitch.Split(IToolMode::EqualsStr, nullptr, &ValStr);
-					if (bIsFirst && !ParseValue(ValStr, FirstValue))
-					{
-						ErrorMessage = FString::Printf(TEXT("Invalid parse of argument %s value: %s."),
-							*FString(FCString::Strlen(InFirstSwitch) - 1, InFirstSwitch),
-							*ValStr);
-						return false;
-					}
-					else if (bIsSecond && !ParseValue(ValStr, SecondValue))
-					{
-						ErrorMessage = FString::Printf(TEXT("Invalid parse of argument %s value: %s."),
-							*FString(FCString::Strlen(InSecondSwitch) - 1, InSecondSwitch),
-							*ValStr);
-						return false;
-					}
-					bFoundFirst |= bIsFirst;
-					bFoundSecond |= bIsSecond;
-				}
-				if (bFoundFirst && bFoundSecond)
-				{
-					OutPairs.Emplace(MoveTemp(FirstValue), MoveTemp(SecondValue));
-					bFoundAnyPair = true;
-					bFoundFirst = false;
-					bFoundSecond = false;
-					FirstValue = TFirst();
-					SecondValue = TSecond();
-				}
-			}
-			if (bFoundFirst != bFoundSecond)
-			{
-				ErrorMessage = FString::Printf(TEXT("%s and %s parameters must be provided as pairs."),
-					*FString(FCString::Strlen(InFirstSwitch) - 1, InFirstSwitch),
-					*FString(FCString::Strlen(InSecondSwitch) - 1, InSecondSwitch));
-				return false;
-			}
-			return bFoundAnyPair;
-		}
-
-		bool ParsePairs(const TCHAR* InFirstSwitch, const TCHAR* InSecondSwitch, const TArray<FString>& Switches, TMap<FGuid, TArray<uint8>>& OutPairs, FString& ErrorMessage)
-		{
-			return BuildPatchTool::IToolMode::ParsePairs<FGuid, TArray<uint8>, TMap<FGuid, TArray<uint8>>>(InFirstSwitch, InSecondSwitch, Switches, OutPairs, ErrorMessage);
-		}
-
-		bool ParsePairs(const TCHAR* InFirstSwitch, const TCHAR* InSecondSwitch, const TArray<FString>& Switches, TArray<TPair<FString, FString>>& OutPairs, FString& ErrorMessage)
-		{
-			return BuildPatchTool::IToolMode::ParsePairs<FString, FString, TArray<TPair<FString, FString>>>(InFirstSwitch, InSecondSwitch, Switches, OutPairs, ErrorMessage);
-		}
-
 		bool ParseValue(const FString& ValueIn, FString& ValueOut)
 		{
 			ValueOut = ValueIn.TrimQuotes();
@@ -301,35 +140,16 @@ namespace BuildPatchTool
 			return false;
 		}
 
-		bool ParseValue(const FString& ValueIn, int32& ValueOut)
-		{
-			if (FCString::IsNumeric(*ValueIn) && !ValueIn.Contains(TEXT("-"), ESearchCase::CaseSensitive))
-			{
-				ValueOut = FCString::Strtoi(*ValueIn, nullptr, 10);
-				return true;
-			}
-			return false;
-		}
-
-		bool ParseValue(const FString& ValueIn, FGuid& ValueOut)
-		{
-			return FGuid::Parse(ValueIn, ValueOut);
-		}
-
-		template<typename TCollection>
-		typename std::enable_if<is_ue_set<TCollection>::value || is_ue_array<TCollection>::value, bool>::type
-			ParseValue(const FString& ValueIn, TCollection& ValueOut)
+		template<typename ContainerType>
+		bool ParseValue(const FString& ValueIn, ContainerType& ValueOut)
 		{
-			TArray<FString> TempArray;
+			TArray<FString> TempOut;
 			ValueOut.Empty();
 			FString CleanValueIn = ValueIn.TrimQuotes();
-			if (CleanValueIn.ParseIntoArray(TempArray, TEXT(","), false) != 0)
+			CleanValueIn.RemoveSpacesInline();
+			if (CleanValueIn.ParseIntoArray(TempOut, TEXT(","), false) != 0)
 			{
-				for (FString& TempElem : TempArray)
-				{
-					TempElem.TrimStartAndEndInline();
-					ValueOut.Add(MoveTemp(TempElem));
-				}
+				ValueOut.Append(MoveTemp(TempOut));
 			}
 			else
 			{
@@ -337,40 +157,17 @@ namespace BuildPatchTool
 				// line and contains no commas.
 				ValueOut.Add(TEXT(""));
 			}
-			return true;
-		}
-
-		bool ParseValue(const FString& ValueIn, TArray<uint8>& ValueOut)
-		{
-			// Here we are expecting a string into bytes, and for now, we are supporting Hex only.
-			const int32 NumChars = ValueIn.Len();
-			// An odd number of characters means the first byte is a padded nibble
-			int32 NumBytes = ValueIn.Len() / 2;
-			if (NumChars % 2 == 1)
-			{
-				++NumBytes;
-			}
-			ValueOut.Empty(NumBytes);
-			ValueOut.AddZeroed(NumBytes);
-			return ValueOut.Num() == 0 || HexToBytes(ValueIn, &ValueOut[0]) == NumBytes;
-		}
 
-		template<typename TEnum>
-		typename std::enable_if<std::is_enum<TEnum>::value, bool>::type
-			ParseValue(const FString& ValueIn, TEnum& ValueOut)
-		{
-			LexFromString(ValueOut, *ValueIn);
-			return ValueOut != TEnum::InvalidOrMax;
+			return true;
 		}
 	};
 
+	typedef TSharedRef<IToolMode> IToolModeRef;
+	typedef TSharedPtr<IToolMode> IToolModePtr;
+
 	class FToolModeFactory
 	{
 	public:
-		typedef TFunction<IToolMode* (IBuildPatchServicesModule&, const TCHAR* /*CommandLine*/)> FToolModeCreateFunc;
-		typedef TMulticastDelegate<void(TMap<FString, FToolModeCreateFunc>&)> FModeRegistryDelegate;
-
-		static IToolMode* Create(IBuildPatchServicesModule& BpsInterface, const TCHAR* CommandLine);
-		static FModeRegistryDelegate& GetModeRegistryDelegate();
+		static IToolModeRef Create(IBuildPatchServicesModule& BpsInterface);
 	};
 }
diff --git a/Engine/Source/Programs/BuildPatchTool/Public/ToolModes/AutomationMode.h b/Engine/Source/Programs/BuildPatchTool/Public/ToolModes/AutomationMode.h
--- a/Engine/Source/Programs/BuildPatchTool/Public/ToolModes/AutomationMode.h
+++ b/Engine/Source/Programs/BuildPatchTool/Public/ToolModes/AutomationMode.h
@@ -3,3 +3,16 @@
 
 #include "CoreMinimal.h"
 #include "Interfaces/ToolMode.h"
+
+#if !UE_BUILD_SHIPPING
+
+namespace BuildPatchTool
+{
+	class FAutomationToolModeFactory
+	{
+	public:
+		static IToolModeRef Create(IBuildPatchServicesModule& BpsInterface);
+	};
+}
+
+#endif // !UE_BUILD_SHIPPING
diff --git a/Engine/Source/Programs/BuildPatchTool/Public/ToolModes/ChunkDeltaOptimiseMode.h b/Engine/Source/Programs/BuildPatchTool/Public/ToolModes/ChunkDeltaOptimiseMode.h
--- a/Engine/Source/Programs/BuildPatchTool/Public/ToolModes/ChunkDeltaOptimiseMode.h
+++ b/Engine/Source/Programs/BuildPatchTool/Public/ToolModes/ChunkDeltaOptimiseMode.h
@@ -3,3 +3,12 @@
 
 #include "CoreMinimal.h"
 #include "Interfaces/ToolMode.h"
+
+namespace BuildPatchTool
+{
+	class FChunkDeltaOptimiseToolModeFactory
+	{
+	public:
+		static IToolModeRef Create(IBuildPatchServicesModule& BpsInterface);
+	};
+}
diff --git a/Engine/Source/Programs/BuildPatchTool/Public/ToolModes/CompactifyMode.h b/Engine/Source/Programs/BuildPatchTool/Public/ToolModes/CompactifyMode.h
--- a/Engine/Source/Programs/BuildPatchTool/Public/ToolModes/CompactifyMode.h
+++ b/Engine/Source/Programs/BuildPatchTool/Public/ToolModes/CompactifyMode.h
@@ -3,3 +3,12 @@
 
 #include "CoreMinimal.h"
 #include "Interfaces/ToolMode.h"
+
+namespace BuildPatchTool
+{
+	class FCompactifyToolModeFactory
+	{
+	public:
+		static IToolModeRef Create(IBuildPatchServicesModule& BpsInterface);
+	};
+}
diff --git a/Engine/Source/Programs/BuildPatchTool/Public/ToolModes/DiffManifestMode.h b/Engine/Source/Programs/BuildPatchTool/Public/ToolModes/DiffManifestMode.h
--- a/Engine/Source/Programs/BuildPatchTool/Public/ToolModes/DiffManifestMode.h
+++ b/Engine/Source/Programs/BuildPatchTool/Public/ToolModes/DiffManifestMode.h
@@ -3,3 +3,12 @@
 
 #include "CoreMinimal.h"
 #include "Interfaces/ToolMode.h"
+
+namespace BuildPatchTool
+{
+	class FDiffManifestToolModeFactory
+	{
+	public:
+		static IToolModeRef Create(IBuildPatchServicesModule& BpsInterface);
+	};
+}
diff --git a/Engine/Source/Programs/BuildPatchTool/Public/ToolModes/EnumerationMode.h b/Engine/Source/Programs/BuildPatchTool/Public/ToolModes/EnumerationMode.h
--- a/Engine/Source/Programs/BuildPatchTool/Public/ToolModes/EnumerationMode.h
+++ b/Engine/Source/Programs/BuildPatchTool/Public/ToolModes/EnumerationMode.h
@@ -3,3 +3,12 @@
 
 #include "CoreMinimal.h"
 #include "Interfaces/ToolMode.h"
+
+namespace BuildPatchTool
+{
+	class FEnumerationToolModeFactory
+	{
+	public:
+		static IToolModeRef Create(IBuildPatchServicesModule& BpsInterface);
+	};
+}
diff --git a/Engine/Source/Programs/BuildPatchTool/Public/ToolModes/MergeManifestMode.h b/Engine/Source/Programs/BuildPatchTool/Public/ToolModes/MergeManifestMode.h
--- a/Engine/Source/Programs/BuildPatchTool/Public/ToolModes/MergeManifestMode.h
+++ b/Engine/Source/Programs/BuildPatchTool/Public/ToolModes/MergeManifestMode.h
@@ -3,3 +3,12 @@
 
 #include "CoreMinimal.h"
 #include "Interfaces/ToolMode.h"
+
+namespace BuildPatchTool
+{
+	class FMergeManifestToolModeFactory
+	{
+	public:
+		static IToolModeRef Create(IBuildPatchServicesModule& BpsInterface);
+	};
+}
diff --git a/Engine/Source/Programs/BuildPatchTool/Public/ToolModes/PackageChunksMode.h b/Engine/Source/Programs/BuildPatchTool/Public/ToolModes/PackageChunksMode.h
--- a/Engine/Source/Programs/BuildPatchTool/Public/ToolModes/PackageChunksMode.h
+++ b/Engine/Source/Programs/BuildPatchTool/Public/ToolModes/PackageChunksMode.h
@@ -3,3 +3,12 @@
 
 #include "CoreMinimal.h"
 #include "Interfaces/ToolMode.h"
+
+namespace BuildPatchTool
+{
+	class FPackageChunksToolModeFactory
+	{
+	public:
+		static IToolModeRef Create(IBuildPatchServicesModule& BpsInterface);
+	};
+}
diff --git a/Engine/Source/Programs/BuildPatchTool/Public/ToolModes/PatchGenerationMode.h b/Engine/Source/Programs/BuildPatchTool/Public/ToolModes/PatchGenerationMode.h
new file mode 100644
--- /dev/null
+++ b/Engine/Source/Programs/BuildPatchTool/Public/ToolModes/PatchGenerationMode.h
@@ -0,0 +1,14 @@
+// Copyright Epic Games, Inc. All Rights Reserved.
+#pragma once
+
+#include "CoreMinimal.h"
+#include "Interfaces/ToolMode.h"
+
+namespace BuildPatchTool
+{
+	class FPatchGenerationToolModeFactory
+	{
+	public:
+		static IToolModeRef Create(IBuildPatchServicesModule& BpsInterface);
+	};
+}
diff --git a/Engine/Source/Programs/BuildPatchTool/Public/ToolModes/VerifyChunksMode.h b/Engine/Source/Programs/BuildPatchTool/Public/ToolModes/VerifyChunksMode.h
--- a/Engine/Source/Programs/BuildPatchTool/Public/ToolModes/VerifyChunksMode.h
+++ b/Engine/Source/Programs/BuildPatchTool/Public/ToolModes/VerifyChunksMode.h
@@ -3,3 +3,12 @@
 
 #include "CoreMinimal.h"
 #include "Interfaces/ToolMode.h"
+
+namespace BuildPatchTool
+{
+	class FVerifyChunksToolModeFactory
+	{
+	public:
+		static IToolModeRef Create(IBuildPatchServicesModule& BpsInterface);
+	};
+}
diff --git a/Engine/Source/Runtime/Online/BuildPatchServices/BuildPatchServices.Build.cs b/Engine/Source/Runtime/Online/BuildPatchServices/BuildPatchServices.Build.cs
--- a/Engine/Source/Runtime/Online/BuildPatchServices/BuildPatchServices.Build.cs
+++ b/Engine/Source/Runtime/Online/BuildPatchServices/BuildPatchServices.Build.cs
@@ -5,9 +5,6 @@ using UnrealBuildTool;
 
 public class BuildPatchServices : ModuleRules
 {
-	[ConfigFile(ConfigHierarchyType.Engine, "BuildPatchServices")]
-	bool bEnableDiskOverflowStore = true;
-
 	public BuildPatchServices(ReadOnlyTargetRules Target) : base(Target)
 	{
 		StaticAnalyzerDisabledCheckers.Add("core.uninitialized.ArraySubscript");
@@ -25,14 +22,6 @@ public class BuildPatchServices : ModuleRules
 				"CoreUObject"
 			}
 		);
-		if (Target.Configuration != UnrealTargetConfiguration.Shipping)
-		{
-			PrivateDependencyModuleNames.AddRange(
-				new string[] {
-					"GoogleTest"
-				}
-			);
-		}
 
 		PrivateDependencyModuleNames.AddRange(
 		new string[] {
@@ -42,64 +31,5 @@ public class BuildPatchServices : ModuleRules
 				"Json",
 			}
 		);
-
-		PrivateIncludePathModuleNames.AddRange(
-			new string[] {
-				"HTTP"
-			}
-		);
-		PrivateIncludePathModuleNames.AddRange(
-			new string[] {
-				"HTTP"
-			}
-		);
-
-		// OpenSSL is needed on desktop platforms to support calculating SHA 256
-		bool bWithOpenSsl = false;
-		if (Target.Platform == UnrealTargetPlatform.Win64
-		 || Target.Platform == UnrealTargetPlatform.Mac
-		 || Target.Platform == UnrealTargetPlatform.Linux)
-		{
-			AddEngineThirdPartyPrivateStaticDependencies(Target, "OpenSSL");
-			bWithOpenSsl = true;
-		}
-
-		// Add a definition to allow BPS source to switch openssl code in and out
-		if (bWithOpenSsl)
-		{
-			PublicDefinitions.Add("WITH_OPENSSL=1");
-		}
-		else
-		{
-			PublicDefinitions.Add("WITH_OPENSSL=0");
-		}
-
-		if (EnableDiskOverflowStore)
-		{
-			PublicDefinitions.Add("ENABLE_PATCH_DISK_OVERFLOW_STORE=1");
-		}
-		else
-		{
-			PublicDefinitions.Add("ENABLE_PATCH_DISK_OVERFLOW_STORE=0");
-		}
-
-		if (EnableDiskOverflowStore)
-		{
-			PublicDefinitions.Add("ENABLE_PATCH_DISK_OVERFLOW_STORE=1");
-		}
-		else
-		{
-			PublicDefinitions.Add("ENABLE_PATCH_DISK_OVERFLOW_STORE=0");
-		}
-
-	}
-
-	protected bool EnableDiskOverflowStore
-	{
-		get
-		{
-			ConfigCache.ReadSettings(DirectoryReference.FromFile(Target.ProjectFile), Target.Platform, this);
-			return bEnableDiskOverflowStore;
-		}
 	}
 }
diff --git a/Engine/Source/Runtime/Online/BuildPatchServices/Private/BuildPatchDelta.cpp b/Engine/Source/Runtime/Online/BuildPatchServices/Private/BuildPatchDelta.cpp
--- a/Engine/Source/Runtime/Online/BuildPatchServices/Private/BuildPatchDelta.cpp
+++ b/Engine/Source/Runtime/Online/BuildPatchServices/Private/BuildPatchDelta.cpp
@@ -11,5 +11,5 @@ BUILDPATCHSERVICES_API FString BuildPatchServices::GetChunkDeltaFilename(const I
 
 BUILDPATCHSERVICES_API IBuildManifestPtr BuildPatchServices::MergeDeltaManifest(const IBuildManifestRef& Manifest, const IBuildManifestRef& Delta)
 {
-	return FBuildMergeManifests::MergeDeltaManifest(StaticCastSharedRef<FBuildPatchAppManifest>(Manifest).Get(), StaticCastSharedRef<FBuildPatchAppManifest>(Delta).Get());
+	return FBuildMergeManifests::MergeDeltaManifest(StaticCastSharedRef<FBuildPatchAppManifest>(Manifest), StaticCastSharedRef<FBuildPatchAppManifest>(Delta));
 }
diff --git a/Engine/Source/Runtime/Online/BuildPatchServices/Private/BuildPatchFeatureLevel.cpp b/Engine/Source/Runtime/Online/BuildPatchServices/Private/BuildPatchFeatureLevel.cpp
new file mode 100644
--- /dev/null
+++ b/Engine/Source/Runtime/Online/BuildPatchServices/Private/BuildPatchFeatureLevel.cpp
@@ -0,0 +1,63 @@
+// Copyright Epic Games, Inc. All Rights Reserved.
+
+#include "BuildPatchFeatureLevel.h"
+
+static_assert((int32)BuildPatchServices::EFeatureLevel::Latest == 18, "Please add support for the extra values to the two functions below.");
+
+BUILDPATCHSERVICES_API const TCHAR* BuildPatchServices::FeatureLevelToString(const BuildPatchServices::EFeatureLevel& FeatureLevel)
+{
+#define CASE_ENUM_TO_STR(Value) case BuildPatchServices::EFeatureLevel::Value: return TEXT("EFeatureLevel::") TEXT(#Value)
+	switch (FeatureLevel)
+	{
+		CASE_ENUM_TO_STR(Original);
+		CASE_ENUM_TO_STR(CustomFields);
+		CASE_ENUM_TO_STR(StartStoringVersion);
+		CASE_ENUM_TO_STR(DataFileRenames);
+		CASE_ENUM_TO_STR(StoresIfChunkOrFileData);
+		CASE_ENUM_TO_STR(StoresDataGroupNumbers);
+		CASE_ENUM_TO_STR(ChunkCompressionSupport);
+		CASE_ENUM_TO_STR(StoresPrerequisitesInfo);
+		CASE_ENUM_TO_STR(StoresChunkFileSizes);
+		CASE_ENUM_TO_STR(StoredAsCompressedUClass);
+		CASE_ENUM_TO_STR(StoresChunkDataShaHashes);
+		CASE_ENUM_TO_STR(StoresPrerequisiteIds);
+		CASE_ENUM_TO_STR(StoredAsBinaryData);
+		CASE_ENUM_TO_STR(VariableSizeChunks);
+		CASE_ENUM_TO_STR(UsesRuntimeGeneratedBuildId);
+		CASE_ENUM_TO_STR(UsesBuildTimeGeneratedBuildId);
+		default: return TEXT("Invalid");
+	}
+#undef CASE_ENUM_TO_STR
+}
+
+BUILDPATCHSERVICES_API bool BuildPatchServices::FeatureLevelFromString(const TCHAR* FeatureLevelString, BuildPatchServices::EFeatureLevel& FeatureLevel)
+{
+#define RETURN_IF_EQUAL(Value) if (FCString::Stricmp(FeatureLevelString, TEXT(#Value)) == 0) { FeatureLevel = BuildPatchServices::EFeatureLevel::Value; return true; }
+	if (FCString::Strnicmp(FeatureLevelString, TEXT("EFeatureLevel::"), 15) == 0)
+	{
+		FeatureLevelString += 15;
+	}
+	RETURN_IF_EQUAL(Original);
+	RETURN_IF_EQUAL(CustomFields);
+	RETURN_IF_EQUAL(StartStoringVersion);
+	RETURN_IF_EQUAL(DataFileRenames);
+	RETURN_IF_EQUAL(StoresIfChunkOrFileData);
+	RETURN_IF_EQUAL(StoresDataGroupNumbers);
+	RETURN_IF_EQUAL(ChunkCompressionSupport);
+	RETURN_IF_EQUAL(StoresPrerequisitesInfo);
+	RETURN_IF_EQUAL(StoresChunkFileSizes);
+	RETURN_IF_EQUAL(StoredAsCompressedUClass);
+	RETURN_IF_EQUAL(StoresChunkDataShaHashes);
+	RETURN_IF_EQUAL(StoresPrerequisiteIds);
+	RETURN_IF_EQUAL(StoredAsBinaryData);
+	RETURN_IF_EQUAL(VariableSizeChunks);
+	RETURN_IF_EQUAL(UsesRuntimeGeneratedBuildId);
+	RETURN_IF_EQUAL(UsesBuildTimeGeneratedBuildId);
+	RETURN_IF_EQUAL(Latest);
+	RETURN_IF_EQUAL(LatestNoChunks);
+	RETURN_IF_EQUAL(LatestJson);
+	RETURN_IF_EQUAL(FirstOptimisedDelta);
+	RETURN_IF_EQUAL(StoresUniqueBuildId);
+	return false;
+#undef RETURN_IF_EQUAL
+}
diff --git a/Engine/Source/Runtime/Online/BuildPatchServices/Private/BuildPatchFileConstructor.cpp b/Engine/Source/Runtime/Online/BuildPatchServices/Private/BuildPatchFileConstructor.cpp
--- a/Engine/Source/Runtime/Online/BuildPatchServices/Private/BuildPatchFileConstructor.cpp
+++ b/Engine/Source/Runtime/Online/BuildPatchServices/Private/BuildPatchFileConstructor.cpp
@@ -2,7 +2,6 @@
 
 #include "BuildPatchFileConstructor.h"
 #include "IBuildManifestSet.h"
-#include "Algo/Transform.h"
 #include "HAL/UESemaphore.h"
 #include "HAL/PlatformFileManager.h"
 #include "HAL/FileManager.h"
@@ -11,7 +10,6 @@
 #include "Misc/ConfigCacheIni.h"
 #include "Misc/Paths.h"
 #include "Misc/ScopeLock.h"
-#include "Misc/ScopeRWLock.h"
 #include "Misc/OutputDeviceRedirector.h"
 #include "BuildPatchServicesPrivate.h"
 #include "Interfaces/IBuildInstaller.h"
@@ -38,9 +36,6 @@
 #include <unistd.h>
 #endif
 
-DECLARE_LOG_CATEGORY_EXTERN(LogBuildPatchFileConstructor, Log, All);
-DEFINE_LOG_CATEGORY(LogBuildPatchFileConstructor);
-
 using namespace BuildPatchServices;
 
 static int32 SleepTimeWhenFileSystemThrottledSeconds = 1;
@@ -84,27 +79,6 @@ namespace FileConstructorHelpers
 		}
 	}
 
-	bool CheckRemainingDiskSpace(const FString& InstallDirectory, uint64 RemainingBytesRequired, uint64& OutAvailableDiskSpace)
-	{
-		bool bContinueConstruction = true;
-		uint64 TotalSize = 0;
-		OutAvailableDiskSpace = 0;
-		if (FPlatformMisc::GetDiskTotalAndFreeSpace(InstallDirectory, TotalSize, OutAvailableDiskSpace))
-		{
-			if (OutAvailableDiskSpace < RemainingBytesRequired)
-			{
-				bContinueConstruction = false;
-			}
-		}
-		else
-		{
-			// If we can't get the disk space free then the most likely reason is the drive is no longer around...
-			bContinueConstruction = false;
-		}
-
-		return bContinueConstruction;
-	}
-
 	uint64 CalculateRequiredDiskSpace(const FBuildPatchAppManifestPtr& CurrentManifest, const FBuildPatchAppManifestRef& BuildManifest, const EInstallMode& InstallMode, const TSet<FString>& InInstallTags)
 	{
 		// Make tags expected
@@ -149,24 +123,6 @@ namespace FileConstructorHelpers
 		}
 		return FMath::Max<int64>(DiskSpaceDeltaPeak, 0);
 	}
-
-	static float ReadRWDurationThresholdSec()
-	{
-		float RWDurationThresholdSec = 4.0;
-		GConfig->GetFloat(TEXT("BuildPatchServices"), TEXT("ConstructionRWDurationThresholdSec"), RWDurationThresholdSec, GEngineIni);
-		return RWDurationThresholdSec;
-	}
-
-	static void LogLowRWDurationTime(const ISpeedRecorder::FRecord& ActivityRecord, const TCHAR* ProcedureType)
-	{
-		// We've seen reports of abnormally poor HDD response times so we log when things are truly bad (e.g. multiple seconds for an access).
-		static const double RWDurationThresholdSec = ReadRWDurationThresholdSec();
-		const double Duration = FStatsCollector::CyclesToSeconds(ActivityRecord.CyclesEnd - ActivityRecord.CyclesStart);
-		if (Duration > RWDurationThresholdSec)
-		{
-			UE_LOG(LogBuildPatchServices, Warning, TEXT("FBuildPatchFileConstructor: %s time of %.3f seconds (%d bytes)"), ProcedureType, Duration, ActivityRecord.Size);
-		}
-	}
 }
 
 
@@ -192,11 +148,9 @@ struct FReadScope
 {
 	ISpeedRecorder::FRecord ActivityRecord;
 	IFileConstructorStat* FileConstructorStat;
-	const TCHAR* LogDetails = nullptr;
-	FReadScope(IFileConstructorStat* InFileConstructorStat, int64 Size, const TCHAR* InLogDetails)
+	FReadScope(IFileConstructorStat* InFileConstructorStat, int64 Size)
 	{
-		LogDetails = InLogDetails;
-		FileConstructorStat = InFileConstructorStat;
+		FileConstructorStat = InFileConstructorStat;		
 		ActivityRecord.CyclesStart = FStatsCollector::GetCycles();
 		ActivityRecord.Size = Size;
 		FileConstructorStat->OnBeforeRead();
@@ -204,7 +158,6 @@ struct FReadScope
 	~FReadScope()
 	{
 		ActivityRecord.CyclesEnd = FStatsCollector::GetCycles();
-		FileConstructorHelpers::LogLowRWDurationTime(ActivityRecord, LogDetails);
 		FileConstructorStat->OnAfterRead(ActivityRecord);
 	}
 };
@@ -222,18 +175,9 @@ enum class EConstructionError : uint8
 	InternalConsistencyError,
 	Aborted,
 	MissingFileInfo,
-	InvalidSymlink,
 	CloseError
 };
 
-enum class EFileStage : uint8
-{
-	Unknown,	 // The file stage is unknown due to a failure (FFileConstructionState::bSuccess is false).
-	Complete,    // The file is complete; it exists and contains the expected data size.
-	Incomplete,  // The file is incomplete; it exists but does not contain the expected data size. Resume the operation.
-	New          // The file is missing; a new file should be created and populated from the start.
-};
-
 
 // Since we can have more than one file in flight, store state here.
 struct FFileConstructionState
@@ -272,9 +216,7 @@ struct FFileConstructionState
 	int32 ConstructionIndex = -1;
 	const FFileManifest* FileManifest;
 	const FString& BuildFilename; // name as it is in the manifest, references the ConstructList in the configuration.
-	const FString NewFilename; // with output path and such
-
-	EFileStage Stage = EFileStage::Unknown;
+	FString NewFilename; // with output path and such
 
 	FFileConstructionState(const FFileManifest* InFileManifest, const FString& InBuildFilename, FString&& InNewFilename) 
 	: FileManifest(InFileManifest)
@@ -290,7 +232,7 @@ struct FFileConstructionState
 
 	void SetAttributes()
 	{
-#if PLATFORM_MAC || PLATFORM_LINUX
+#if PLATFORM_MAC
 		if (bSuccess && 
 			EnumHasAllFlags(FileManifest->FileMetaFlags, EFileMetaFlags::UnixExecutable))
 		{
@@ -339,8 +281,14 @@ public:
 	// The resume ids that we loaded from disk
 	TSet<FString> LoadedResumeIds;
 
-	// Maps each resumed file to its current stage: Complete or Incomplete. Files not found on disk or treated as New are not present in this map.
-	TMap<FString, EFileStage> FilesStage;
+	// The set of files that were started
+	TSet<FString> FilesStarted;
+
+	// The set of files that were completed, determined by expected file size
+	TSet<FString> FilesCompleted;
+
+	// The set of files that exist but are not able to assume resumable
+	TSet<FString> FilesIncompatible;
 
 	// Whether we have any resume data for this install
 	bool bHasResumeData = false;
@@ -405,24 +353,11 @@ public:
 		}
 	}
 
-	/**
-	 * Deletes the resume data
-	 */
-	void RemoveResumeData()
-	{
-		if (bResumeEnabled)
-		{
-			bHasResumeData = false;
-			FileSystem->DeleteFile(*ResumeDataFilename);
-		}
-	}
-
 	/**
 	 * Checks whether the file was completed during last install attempt and adds it to FilesCompleted if so
 	 * @param Filename    The filename to check
-	 * @return A flag indicating that the file is incompatible and should be removed and constracted from scratch.
 	 */
-	bool CheckFile(const FString& Filename)
+	void CheckFile(const FString& Filename)
 	{
 		// If we had resume data, check if this file might have been resumable
 		if (bHasResumeData)
@@ -439,27 +374,25 @@ public:
 				if (NewFileManifest && bFileExists)
 				{
 					const uint64 UnsignedDiskFileSize = DiskFileSize;
-					if (UnsignedDiskFileSize == NewFileManifest->FileSize)
+					if (UnsignedDiskFileSize > 0 && UnsignedDiskFileSize <= NewFileManifest->FileSize)
 					{
-						FilesStage.Add(Filename, EFileStage::Complete);
+						FilesStarted.Add(Filename);
 					}
-					else if (UnsignedDiskFileSize > 0 && UnsignedDiskFileSize < NewFileManifest->FileSize)
+					if (UnsignedDiskFileSize == NewFileManifest->FileSize)
 					{
-						FilesStage.Add(Filename, EFileStage::Incomplete);
+						FilesCompleted.Add(Filename);
 					}
-					else if (UnsignedDiskFileSize > NewFileManifest->FileSize)
+					if (UnsignedDiskFileSize > NewFileManifest->FileSize)
 					{
-						return true;
+						FilesIncompatible.Add(Filename);
 					}
-					// If the file size is zero, we just treat it as EFileStage::New
 				}
 			}
 			else if (bFileExists)
 			{
-				return true;
+				FilesIncompatible.Add(Filename);
 			}
 		}
-		return false;
 	}
 };
 
@@ -1507,14 +1440,6 @@ FBuildPatchFileConstructor::FBuildPatchFileConstructor(
 
 		FFileToConstruct FileToConstruct;
 		FileToConstruct.FileManifest = FileManifest;
-		if (Configuration.bUseSHA1StageFilenames)
-		{
-			FBuildPatchUtils::SHAToBase32(FileManifest->SHA1Hash, FileToConstruct.StagingFilename);
-		}
-		else
-		{
-			FileToConstruct.StagingFilename = FileToConstructName;
-		}
 
 		// If we are missing the file manifest, we will fail to install when we get to the file. However,
 		// we guarantee a 1:1 mapping with the arrays we are filling here so we use invalid data for those
@@ -1680,6 +1605,7 @@ FBuildPatchFileConstructor::FBuildPatchFileConstructor(
 
 		Threads[ThreadIndex]->RunTask([this, ThreadIndex]() { GenericThreadFn(ThreadIndex); });
 	}
+
 }
 
 FBuildPatchFileConstructor::~FBuildPatchFileConstructor()
@@ -1799,25 +1725,35 @@ void FBuildPatchFileConstructor::QueueGenericThreadTask(int32 ThreadIndex, ICons
 
 void FBuildPatchFileConstructor::GenericThreadFn(int32 ThreadIndex)
 {
-	bool bLastShouldAbort = false;
-	do
+	for (;;)
 	{
 		ThreadWakeups[ThreadIndex]->Wait();
 
 		TArray<IConstructorChunkSource::FRequestProcessFn> GrabbedJobs;
+		
 		{
-			const UE::TScopeLock Lock(ThreadJobPostingLocks[ThreadIndex]);
+			ThreadJobPostingLocks[ThreadIndex].Lock();
 			GrabbedJobs = MoveTemp(ThreadJobPostings[ThreadIndex]);
+			ThreadJobPostingLocks[ThreadIndex].Unlock();
+		}
+
+		if (bShouldAbort)
+		{
+			for (IConstructorChunkSource::FRequestProcessFn& AbortJob : GrabbedJobs)
+			{
+				AbortJob(true);
+			}
+			WakeUpDispatchThreadEvent->Trigger();
+			break;
 		}
 
-		bLastShouldAbort = bShouldAbort;
-		for (const IConstructorChunkSource::FRequestProcessFn& Job : GrabbedJobs)
+		for (IConstructorChunkSource::FRequestProcessFn& Job : GrabbedJobs)
 		{
-			Job(bLastShouldAbort);
+			Job(false);
 		}
 
 		WakeUpDispatchThreadEvent->Trigger();
-	} while (!bLastShouldAbort);
+	}
 
 	ThreadCompleteEvents[ThreadIndex]->Trigger();
 }
@@ -1888,7 +1824,7 @@ bool FBuildPatchFileConstructor::HarvestChunksForCompletedFile(const FString& In
 					continue;
 				}
 
-				int32 ChunkSize = Configuration.ManifestSet->GetChunkInfo(FileChunk)->DataSizeUncompressed;
+				int32 ChunkSize = Configuration.ManifestSet->GetChunkInfo(FileChunk)->WindowSize;
 
 				ChunksFromFileWeNeed.Add({FileChunk, LastUsageIndex, NextUsageIndex, ChunkSize});
 			}
@@ -2021,7 +1957,7 @@ void FBuildPatchFileConstructor::Run()
 	for (const TPair<FGuid, EConstructorChunkLocation>& ChunkLocation : ChunkLocations)
 	{
 		FChunkInfo const* ChunkInfo = Configuration.ManifestSet->GetChunkInfo(ChunkLocation.Key);
-		WindowSizes.FindOrAdd(ChunkInfo->DataSizeUncompressed)++;
+		WindowSizes.FindOrAdd(ChunkInfo->WindowSize)++;
 
 		if (ChunkLocation.Value == EConstructorChunkLocation::Cloud)
 		{
@@ -2090,32 +2026,26 @@ void FBuildPatchFileConstructor::Run()
 
 		ResumeData.InitResume();
 
-		// Check the stage of resumed files and remove any incompatible files.
-		if (ResumeData.bHasResumeData && !Configuration.bUseSHA1StageFilenames)
+		// Remove incompatible files
+		if (ResumeData.bHasResumeData)
 		{
 			for (const FString& FileToConstruct : Configuration.ConstructList)
 			{
-				const bool bFileIncompatible = ResumeData.CheckFile(FileToConstruct);
+				ResumeData.CheckFile(FileToConstruct);
+				const bool bFileIncompatible = ResumeData.FilesIncompatible.Contains(FileToConstruct);
 				if (bFileIncompatible)
 				{
-					UE_LOG(LogBuildPatchFileConstructor, Log, TEXT("Deleting incompatible stage file %s"), *FileToConstruct);
+					GLog->Logf(TEXT("FBuildPatchFileConstructor: Deleting incompatible stage file %s"), *FileToConstruct);
 					FileSystem->DeleteFile(*(Configuration.StagingDirectory / FileToConstruct));
 				}
 			}
 
 			if (Configuration.bConstructInMemory)
 			{
-				UE_LOG(LogBuildPatchServices, Log, TEXT("Removing incomplete entries from resume data: construct in memory"));
-				ResumeData.FilesStage = ResumeData.FilesStage.FilterByPredicate([](const TPair<FString, EFileStage>& Entry) { return Entry.Value != EFileStage::Incomplete; });
+				UE_LOG(LogBuildPatchServices, Log, TEXT("Emptying Resume Data FilesStarted: construct in memory"));
+				ResumeData.FilesStarted.Empty();
 			}
-		}
 
-		// Uninstall should not have any files to construct
-		const bool bIsUninstall = Configuration.ConstructList.IsEmpty();
-		if (ResumeData.bHasResumeData && !bIsUninstall)
-		{
-			// For uninstall, we will remove staged files at an upper level
-			RemoveUnnecessaryStagedFiles(Configuration.StagingDirectory);
 		}
 
 		// Save for started versions
@@ -2137,12 +2067,6 @@ void FBuildPatchFileConstructor::Run()
 	{
 		FileConstructorStat->OnResumeCompleted();
 	}
-
-	if (InstallerError->HasError() == false)
-	{
-		ResumeData.RemoveResumeData();
-	}
-
 	FileConstructorStat->OnConstructionCompleted();
 }
 
@@ -2156,21 +2080,11 @@ uint64 FBuildPatchFileConstructor::GetAvailableDiskSpace()
 	return AvailableDiskSpace.load(std::memory_order_relaxed);
 }
 
-bool FBuildPatchFileConstructor::HasCorruptFiles() const
-{
-	return !CorruptFiles.IsEmpty();
-}
-
 FBuildPatchFileConstructor::FOnBeforeDeleteFile& FBuildPatchFileConstructor::OnBeforeDeleteFile()
 {
 	return BeforeDeleteFileEvent;
 }
 
-FBuildPatchFileConstructor::FOnFileIsStaged& FBuildPatchFileConstructor::OnFileIsStaged()
-{
-	return FileIsStagedEvent;
-}
-
 void FBuildPatchFileConstructor::CountBytesProcessed( const int64& ByteCount )
 {
 	ByteProcessed += ByteCount;
@@ -2202,9 +2116,6 @@ uint64 FBuildPatchFileConstructor::CalculateInProgressDiskSpaceRequired(const FF
 {
 	if (Configuration.InstallMode == EInstallMode::DestructiveInstall)
 	{
-		TSet<FString> MissedCurrentFiles; // Files that were not installed, possibly due to filtering by tags.
-		Configuration.ManifestSet->GetMissedFiles(Configuration.InstallDirectory, MissedCurrentFiles);
-
 		// The simplest method will be to run through each high level file operation, tracking peak disk usage delta.
 
 		// We know we need enough space to finish writing this file
@@ -2243,8 +2154,8 @@ uint64 FBuildPatchFileConstructor::CalculateInProgressDiskSpaceRequired(const FF
 			{
 				DiskSpaceDeltaPeak = DiskSpaceDelta;
 			}
-			// Then we can remove the currently existing file if it was previously created.
-			if (OldFileManifest && !MissedCurrentFiles.Contains(Configuration.ConstructList[ConstructionIndex]))
+			// Then we can remove the current existing file.
+			if (OldFileManifest)
 			{
 				DiskSpaceDelta -= OldFileManifest->FileSize;
 			}
@@ -2260,15 +2171,19 @@ uint64 FBuildPatchFileConstructor::CalculateInProgressDiskSpaceRequired(const FF
 	}
 }
 
-uint64 FBuildPatchFileConstructor::CalculateDiskSpaceRequirementsWithDeleteDuringInstall(int64 InProgressFileSize)
+uint64 FBuildPatchFileConstructor::CalculateDiskSpaceRequirementsWithDeleteDuringInstall()
 {
-	checkf(ChunkDbSource, TEXT("CalculateDiskSpaceRequirementsWithDeleteDuringInstall can't be called without valid ChunkDbSource"));
+	if (ChunkDbSource == nullptr)
+	{
+		// invalid use.
+		return 0;
+	}
 
 	// These are the sizes at after each file that we _started_ with. This is the size after retirement for the
 	// file at those positions.
 	TArray<uint64> ChunkDbSizesAtPosition;
-	ChunkDbSizesAtPosition.Reserve(FileCompletionPositions.Num());
-	const uint64 TotalChunkDbSize = ChunkDbSource->GetChunkDbSizesAtIndexes(FileCompletionPositions, ChunkDbSizesAtPosition);
+	uint64 TotalChunkDbSize = ChunkDbSource->GetChunkDbSizesAtIndexes(FileCompletionPositions, ChunkDbSizesAtPosition);
+
 	// Strip off the files we've completed.
 	int32 CompletedFileCount = NextIndexToConstruct.load(std::memory_order_acquire);
 
@@ -2278,12 +2193,12 @@ uint64 FBuildPatchFileConstructor::CalculateDiskSpaceRequirementsWithDeleteDurin
 	{
 		CompletedFileCount--;
 	}
-	const uint64 MaxDiskSize = FBuildPatchUtils::CalculateDiskSpaceRequirementsWithDeleteDuringInstall(
-		Configuration.ConstructList, CompletedFileCount, InProgressFileSize, Configuration.ManifestSet, ChunkDbSizesAtPosition,
-		Configuration.InstallDirectory, Configuration.InstallMode);
+
+	uint64 MaxDiskSize = FBuildPatchUtils::CalculateDiskSpaceRequirementsWithDeleteDuringInstall(
+		Configuration.ConstructList, CompletedFileCount, Configuration.ManifestSet, ChunkDbSizesAtPosition, TotalChunkDbSize);
 
 	// Strip off the data we already have on disk.
-	uint64 PostDlSize = 0; // Zero means no additional free space is needed for action, as it will be freed by removing chunks.
+	uint64 PostDlSize = 0;
 	if (MaxDiskSize > TotalChunkDbSize)
 	{
 		PostDlSize = MaxDiskSize - TotalChunkDbSize;
@@ -2516,8 +2431,6 @@ IConstructorChunkSource::FRequestProcessFn FBuildPatchFileConstructor::CreateWri
 
 		ActivityRecord.Size = WriteBuffer.GetSize();
 		ActivityRecord.CyclesEnd = FStatsCollector::GetCycles();
-		FileConstructorHelpers::LogLowRWDurationTime(ActivityRecord, TEXT("Write"));
-
 		FileConstructorStat->OnAfterWrite(ActivityRecord);
 		
 		Batch->bIsFinished.store(true, std::memory_order_release);
@@ -2607,7 +2520,6 @@ void FBuildPatchFileConstructor::CompleteReadBatch(const FFileManifest& FileMani
 				Batch.ErrorContextGuid = FileManifest.ChunkParts[Batch.StartChunkPartIndex + i].Guid;
 			}
 		}
-		RemainingChunkCount = ChunkReferenceTracker->GetRemainingChunkCount();
 	}
 
 	if (Batch.Error == EConstructionError::None)
@@ -2647,7 +2559,7 @@ void FBuildPatchFileConstructor::StartReadBatch(FFileConstructionState& CurrentF
 			FRequestInfo& Request = Batch.Requests.FindOrAdd(ChunkPart.Guid);
 			Request.Guid = ChunkPart.Guid;
 			Request.Batch = &Batch;
-			Request.ChunkSize = Configuration.ManifestSet->GetChunkInfo(ChunkPart.Guid)->DataSizeUncompressed;
+			Request.ChunkSize = Configuration.ManifestSet->GetChunkInfo(ChunkPart.Guid)->WindowSize;
 			Request.File = &CurrentFile;
 			Request.ChunkIndexInFile = EndChunkIdx;
 			
@@ -2793,19 +2705,6 @@ void FBuildPatchFileConstructor::StartReadBatch(FFileConstructionState& CurrentF
 				IConstructorChunkSource::FChunkRequestCompleteDelegate::CreateRaw(this, &FBuildPatchFileConstructor::RequestCompletedFn));
 						
 			QueueGenericThreadTask(ThreadAssignments[ChunkLocation], MoveTemp(RequestProcess));
-
-			if (ThisChunkSource == ChunkDbSource)
-			{
-				IConstructorChunkSource::FRequestProcessFn TruncateTask = [this](bool bIsAborted) {
-					if (!bIsAborted)
-					{
-						ChunkDbSource->TryTruncateFiles(RemainingChunkCount);
-					}
-				};
-				// Execute the truncating on the same thread that reads chunks from chunkdb to make archive access lock-free.
-				QueueGenericThreadTask(ThreadAssignments[ChunkLocation], MoveTemp(TruncateTask));
-
-			}
 		}
 	}
 
@@ -2817,50 +2716,6 @@ void FBuildPatchFileConstructor::WakeUpDispatch()
 	WakeUpDispatchThreadEvent->Trigger();
 }
 
-EFileStage FBuildPatchFileConstructor::GetFileStage(const FFileConstructionState& CurrentFile, const FResumeData& ResumeData) const
-{
-	if (!CurrentFile.bSuccess)
-	{
-		return EFileStage::Unknown;
-	}
-
-	if (Configuration.bUseSHA1StageFilenames)
-	{
-		// We have to check file existence here
-		int64 DiskFileSize;
-		if (const bool bFileExists = FileSystem->GetFileSize(*CurrentFile.NewFilename, DiskFileSize))
-		{
-			const uint64 UnsignedDiskFileSize = DiskFileSize;
-			if (UnsignedDiskFileSize == CurrentFile.FileManifest->FileSize)
-			{
-				return EFileStage::Complete;
-			}
-			
-			if (UnsignedDiskFileSize > 0 && UnsignedDiskFileSize < CurrentFile.FileManifest->FileSize) // we only want to resume if the file hasn't been completed.
-			{
-				return EFileStage::Incomplete;
-			}
-			
-			if (UnsignedDiskFileSize > CurrentFile.FileManifest->FileSize)
-			{
-				UE_LOG(LogBuildPatchFileConstructor, Log, TEXT("Deleting incompatible stage file %s (%s)"), *CurrentFile.BuildFilename, *CurrentFile.NewFilename);
-				FileSystem->DeleteFile(*CurrentFile.NewFilename);
-			}
-			// A zero-size file is treated as EFileStage::New.
-		}
-	}
-	else if(const EFileStage * StagePtr = ResumeData.FilesStage.Find(CurrentFile.BuildFilename))
-	{
-		// Can only be EFileStage::Complete or EFileStage::Incomplete.
-		const EFileStage Stage = *StagePtr; // Avoid SA C6269
-		ensure((Stage == EFileStage::Complete) || (Stage == EFileStage::Incomplete));
-		return Stage;
-	}
-
-	// Treat the file as new.
-	return EFileStage::New;
-}
-
 void FBuildPatchFileConstructor::InitFile(FFileConstructionState& CurrentFile, const FResumeData& ResumeData)
 {
 	if (!CurrentFile.bSuccess)
@@ -2868,10 +2723,14 @@ void FBuildPatchFileConstructor::InitFile(FFileConstructionState& CurrentFile, c
 		return;
 	}
 
+	const int64 FileSize = CurrentFile.FileManifest->FileSize;
+
+	// Check resume status for this file.
+	const bool bFilePreviouslyComplete = ResumeData.FilesCompleted.Contains(CurrentFile.BuildFilename);	
+
 	// Construct or skip the file.
-	if (CurrentFile.Stage == EFileStage::Complete)
+	if (bFilePreviouslyComplete)
 	{
-		const int64 FileSize = CurrentFile.FileManifest->FileSize;
 		CountBytesProcessed(FileSize);
 
 		UE_LOG(LogBuildPatchServices, Display, TEXT("Skipping completed file %s"), *CurrentFile.BuildFilename);
@@ -2885,26 +2744,26 @@ void FBuildPatchFileConstructor::InitFile(FFileConstructionState& CurrentFile, c
 				break;
 			}
 		}
-		RemainingChunkCount = ChunkReferenceTracker->GetRemainingChunkCount();
 
-		// Can't skip if we are failing.
-		CurrentFile.bSkippedConstruction = CurrentFile.bSuccess;
+		CurrentFile.bSkippedConstruction = true;
 		return;
 	}
 
-	if (CurrentFile.bSuccess &&
+	if (!CurrentFile.bSuccess &&
 		!CurrentFile.FileManifest->SymlinkTarget.IsEmpty())
 	{
-		// Symlinks are handled separately from file construction
-		if (Configuration.bRejectSymlinks)
+#if PLATFORM_MAC
+		CurrentFile.bSuccess = true;
+		if (!Configuration.bInstallToMemory && !Configuration.bConstructInMemory)
 		{
-			CurrentFile.bSuccess = false;
-			CurrentFile.ConstructionError = EConstructionError::InvalidSymlink;
+			CurrentFile.bSuccess = symlink(TCHAR_TO_UTF8(*CurrentFile.FileManifest->SymlinkTarget), TCHAR_TO_UTF8(*CurrentFile.NewFilename)) == 0;
 		}
-
-		// Can't skip if we are failing.
-		CurrentFile.bSkippedConstruction = CurrentFile.bSuccess;
-		return;
+		CurrentFile.bSkippedConstruction = true;
+#else
+		const bool bSymlinkNotImplemented = false;
+		check(bSymlinkNotImplemented);
+		CurrentFile.bSuccess = false;
+#endif
 	}
 
 	if (Configuration.bInstallToMemory || Configuration.bConstructInMemory)
@@ -2932,18 +2791,13 @@ bool FBuildPatchFileConstructor::HandleInitialDiskSizeCheck(const FFileManifest&
 		bInitialDiskSizeCheck = true;
 
 		// Normal operation can just use the classic calculation
-		uint64 LocalDiskSpaceRequired = 0;
+		uint64 LocalDiskSpaceRequired = CalculateInProgressDiskSpaceRequired(FileManifest, StartPosition);
 
 		// If we are delete-during-install this gets more complicated because we'll be freeing up
 		// space as we add.
-		if (Configuration.bDeleteChunkDBFilesAfterUse && ChunkDbSource)
-		{
-			LocalDiskSpaceRequired = CalculateDiskSpaceRequirementsWithDeleteDuringInstall(StartPosition);
-		}
-		else
+		if (Configuration.bDeleteChunkDBFilesAfterUse)
 		{
-			// Normal operation can just use the classic calculation
-			LocalDiskSpaceRequired = CalculateInProgressDiskSpaceRequired(FileManifest, StartPosition);
+			LocalDiskSpaceRequired = CalculateDiskSpaceRequirementsWithDeleteDuringInstall();
 		}
 
 		bool bHaveSufficientFreeSpace = false;
@@ -2971,7 +2825,7 @@ bool FBuildPatchFileConstructor::HandleInitialDiskSizeCheck(const FFileManifest&
 			UE_LOG(LogBuildPatchServices, Error, TEXT("Insufficient disk space on initial check."));
 			InstallerError->SetError(
 				EBuildPatchInstallError::OutOfDiskSpace,
-				StartPosition != 0 ? DiskSpaceErrorCodes::DuringInstallation : DiskSpaceErrorCodes::InitialSpaceCheck,
+				DiskSpaceErrorCodes::InitialSpaceCheck,
 				0,
 				BuildPatchServices::GetDiskSpaceMessage(Configuration.InstallDirectory, LocalDiskSpaceRequired, LocalDiskSpaceAvailable));
 			return false;
@@ -3021,7 +2875,7 @@ void FBuildPatchFileConstructor::ResumeFile(FFileConstructionState& FileToResume
 			ReadBuffer.SetNumUninitialized(ChunkPart.Size, EAllowShrinking::No);
 
 			{
-				FReadScope _(FileConstructorStat, ChunkPart.Size, TEXT("ReadResume"));
+				FReadScope _(FileConstructorStat, ChunkPart.Size);
 				NewFileReader->Serialize(ReadBuffer.GetData(), ChunkPart.Size);
 			}
 
@@ -3051,7 +2905,6 @@ void FBuildPatchFileConstructor::ResumeFile(FFileConstructionState& FileToResume
 			break;
 		}
 	}
-	RemainingChunkCount = ChunkReferenceTracker->GetRemainingChunkCount();
 
 	NewFileReader->Close();
 	FileToResume.bIsResumedFile = true;
@@ -3094,7 +2947,7 @@ void FBuildPatchFileConstructor::OpenFileToConstruct(FFileConstructionState& Cur
 }
 
 
-void FBuildPatchFileConstructor::CompleteConstructedFile(FFileConstructionState& CurrentFile, FResumeData& ResumeData)
+void FBuildPatchFileConstructor::CompleteConstructedFile(FFileConstructionState& CurrentFile)
 {
 	UE_LOG(LogBuildPatchServices, Display, TEXT("Completing: %s"), *CurrentFile.BuildFilename);
 	bShouldLogNextDependencyWait = true;
@@ -3149,7 +3002,7 @@ void FBuildPatchFileConstructor::CompleteConstructedFile(FFileConstructionState&
 
 			FSHAHash HashValue;
 			CurrentFile.HashState.GetHash(HashValue.Hash);
-			CurrentFile.bSuccess = HashValue == CurrentFile.FileManifest->SHA1Hash;
+			CurrentFile.bSuccess = HashValue == CurrentFile.FileManifest->FileHash;
 			if (!CurrentFile.bSuccess)
 			{
 				CurrentFile.ConstructionError = EConstructionError::OutboundDataError;
@@ -3292,8 +3145,13 @@ void FBuildPatchFileConstructor::CompleteConstructedFile(FFileConstructionState&
 		{
 		case EConstructionError::OutboundDataError:
 			{
+				// Only report if the first error
+				if (bReportAnalytic)
+				{
+					InstallerAnalytics->RecordConstructionError(CurrentFile.BuildFilename, INDEX_NONE, TEXT("Serialised Verify Fail"));
+				}
 				UE_LOG(LogBuildPatchServices, Error, TEXT("FBuildPatchFileConstructor: Verify failed after constructing %s"), *CurrentFile.BuildFilename);
-				CorruptFiles.Add(CurrentFile.BuildFilename);
+				InstallerError->SetError(EBuildPatchInstallError::FileConstructionFail, ConstructionErrorCodes::OutboundCorrupt);
 				break;
 			}
 		case EConstructionError::FailedWrite:
@@ -3310,7 +3168,7 @@ void FBuildPatchFileConstructor::CompleteConstructedFile(FFileConstructionState&
 				uint64 LocalDiskSpaceRequired;
 				if (Configuration.bDeleteChunkDBFilesAfterUse)
 				{
-					LocalDiskSpaceRequired = CalculateDiskSpaceRequirementsWithDeleteDuringInstall(CurrentFile.Progress);
+					LocalDiskSpaceRequired = CalculateDiskSpaceRequirementsWithDeleteDuringInstall();
 				}
 				else
 				{
@@ -3386,16 +3244,6 @@ void FBuildPatchFileConstructor::CompleteConstructedFile(FFileConstructionState&
 				InstallerError->SetError(EBuildPatchInstallError::FileConstructionFail, ConstructionErrorCodes::TrackingError);
 				break;
 			}
-		case EConstructionError::InvalidSymlink:
-			{
-				if (bReportAnalytic)
-				{
-					InstallerAnalytics->RecordConstructionError(CurrentFile.BuildFilename, INDEX_NONE, TEXT("Symlink present when not allowed"));
-				}
-				UE_LOG(LogBuildPatchFileConstructor, Error, TEXT("File manifest includes symlinks when they are forbidden %s"), *CurrentFile.BuildFilename);
-				InstallerError->SetError(EBuildPatchInstallError::FileConstructionFail, ConstructionErrorCodes::SymlinkCreateForbidden);
-				break;
-			}
 		case EConstructionError::InternalConsistencyError:
 			{
 				if (bReportAnalytic)
@@ -3419,15 +3267,6 @@ void FBuildPatchFileConstructor::CompleteConstructedFile(FFileConstructionState&
 		case EConstructionError::FailedInitialSizeCheck:
 			{
 				// Error already set back when we had the info.
-
-				// If we failed the initial disk space check before we did anything,
-				// delete the resume data so that we don't leave any trash. We know there wasn't any existing resume
-				// data because we failed on the first file at the beginning of the file.
-				if (CurrentFile.ConstructionIndex == 0 &&
-					CurrentFile.StartPosition == 0)
-				{
-					ResumeData.RemoveResumeData();
-				}
 				break;
 			}
 		case EConstructionError::Aborted:
@@ -3464,16 +3303,8 @@ void FBuildPatchFileConstructor::CompleteConstructedFile(FFileConstructionState&
 			}
 		}
 
-		// Before aborting the installation, we need to identify and process all corrupted files.
-		if (CurrentFile.ConstructionError != EConstructionError::OutboundDataError)
-		{
-			// Stop trying to install further files.
-			Abort();
-		}
-	}
-	else
-	{
-		OnFileIsStaged().Broadcast(CurrentFile.BuildFilename);
+		// Stop trying to install further files.
+		Abort();
 	}
 }
 
@@ -3488,7 +3319,7 @@ void FBuildPatchFileConstructor::StartFile(FFileConstructionState& CurrentFile,
 		return;
 	}
 
-	if (CurrentFile.Stage == EFileStage::Incomplete && CurrentFile.bSuccess)
+	if (ResumeData.FilesStarted.Contains(CurrentFile.BuildFilename))
 	{
 		ResumeFile(CurrentFile);
 
@@ -3533,7 +3364,7 @@ void FBuildPatchFileConstructor::StartFile(FFileConstructionState& CurrentFile,
 
 }
 
-void FBuildPatchFileConstructor::ConstructFiles(FResumeData& ResumeData)
+void FBuildPatchFileConstructor::ConstructFiles(const FResumeData& ResumeData)
 {
 	MaxWriteBatchSize = FFileConstructorConfig::DefaultIOBatchSizeMB;
 	if (GConfig->GetInt(TEXT("Portal.BuildPatch"), TEXT("ConstructorIOBatchSizeMB"), (int32&)MaxWriteBatchSize, GEngineIni))
@@ -3696,7 +3527,7 @@ void FBuildPatchFileConstructor::ConstructFiles(FResumeData& ResumeData)
 				if (File->NextChunkPartToRead == File->FileManifest->ChunkParts.Num() &&
 					File->OutstandingBatches == 0)
 				{
-					CompleteConstructedFile(*File, ResumeData);
+					CompleteConstructedFile(*File);
 
 					UE_LOG(LogBuildPatchServices, Verbose, TEXT("Completed file: %s, New ActiveFileCount = %d"), *File->BuildFilename, ActiveFiles.Num() - 1);
 
@@ -3898,66 +3729,49 @@ void FBuildPatchFileConstructor::ConstructFiles(FResumeData& ResumeData)
 
 								if (!bDelayForDependencies)
 								{
-									const FString &StagingFileName = ConstructionList[IndexToConstruct].StagingFilename;
-									TUniquePtr<FFileConstructionState> AnotherFileCandidate = MakeUnique<FFileConstructionState>(
-										ConstructionList[IndexToConstruct].FileManifest,
-										Configuration.ConstructList[IndexToConstruct],
-										Configuration.StagingDirectory / StagingFileName);
-
-									AnotherFileCandidate->Stage = GetFileStage(*AnotherFileCandidate, ResumeData);
-									
-									// The complete file will modify ChunkReferenceTracker, which should be done sequentially after processing active files.
-									const bool bDelayForCompleteFile = AnotherFileCandidate->Stage == EFileStage::Complete && ActiveFiles.Num() > 0;
-
-									if (!bDelayForCompleteFile)
-									{
-										// Now commit this file since we are starting it.
-										NextIndexToConstruct.fetch_add(1, std::memory_order_acq_rel);
-
-										// This log has to be here for tests to succeed, but we don't want spam so we make it Verbose.
-										UE_LOG(LogBuildPatchFileConstructor, Log, TEXT("Constructing file %s (%s)"), *Configuration.ConstructList[IndexToConstruct], *StagingFileName);
+									// Now commit this file since we are starting it.
+									NextIndexToConstruct.fetch_add(1, std::memory_order_acq_rel);
 
-										TUniquePtr<FFileConstructionState> AnotherFile = MoveTemp(AnotherFileCandidate);
+									FString StagingFileName = Configuration.StagingDirectory / Configuration.ConstructList[IndexToConstruct];
+									TUniquePtr<FFileConstructionState> AnotherFile = MakeUnique<FFileConstructionState>(
+										ConstructionList[IndexToConstruct].FileManifest, 
+										Configuration.ConstructList[IndexToConstruct], 
+										MoveTemp(StagingFileName));
 
-										AnotherFile->BaseReferenceIndex = 0;
-										if (IndexToConstruct)
-										{
-											AnotherFile->BaseReferenceIndex = FileCompletionPositions[IndexToConstruct - 1];
-										}
+									AnotherFile->BaseReferenceIndex = 0;
+									if (IndexToConstruct)
+									{
+										AnotherFile->BaseReferenceIndex = FileCompletionPositions[IndexToConstruct-1];
+									}
 
-										AnotherFile->ConstructionIndex = IndexToConstruct;
+									AnotherFile->ConstructionIndex = IndexToConstruct;
 
-										TRACE_BOOKMARK(TEXT("Starting File: %s [%u]"), *AnotherFile->BuildFilename, AnotherFile->FileManifest->ChunkParts.Num());
+									TRACE_BOOKMARK(TEXT("Starting File: %s [%u]"), *AnotherFile->BuildFilename, AnotherFile->FileManifest->ChunkParts.Num());
 
-										UE_LOG(LogBuildPatchServices, Display, TEXT("Starting File: %s [%s bytes, %d chunks], New ActiveFileCount = %d"),
-											*AnotherFile->BuildFilename,
-											*FormatNumber(AnotherFile->FileManifest->FileSize),
-											AnotherFile->FileManifest->ChunkParts.Num(),
-											ActiveFiles.Num() + 1);
+									UE_LOG(LogBuildPatchServices, Display, TEXT("Starting File: %s [%s bytes, %d chunks], New ActiveFileCount = %d"), 
+										   *AnotherFile->BuildFilename, 
+										   *FormatNumber(AnotherFile->FileManifest->FileSize),
+										   AnotherFile->FileManifest->ChunkParts.Num(),
+										   ActiveFiles.Num() + 1);
 
-										StartFile(*AnotherFile.Get(), ResumeData);
+									StartFile(*AnotherFile.Get(), ResumeData);
 
-										if (AnotherFile->bSkippedConstruction)
-										{
-											// Nothing else needs to happen with this file - we'll loop around to try for another.
-											bCheckForAnotherRead = true;
-										}
-										else
-										{
-											// Only start the file if there's not currently an active file - otherwise we
-											// do it when the current one finishes.
-											if (ActiveFileCount == 0)
-											{
-												FileConstructorStat->OnFileStarted(AnotherFile->BuildFilename, AnotherFile->FileManifest->FileSize);
-											}
-
-											FileToStart = AnotherFile.Get();
-											ActiveFiles.Push(MoveTemp(AnotherFile));
-										}
+									if (AnotherFile->bSkippedConstruction)
+									{
+										// Nothing else needs to happen with this file - we'll loop around to try for another.
+										bCheckForAnotherRead = true;
 									}
 									else
 									{
-										UE_LOG(LogBuildPatchServices, Verbose, TEXT("Delaying %s due to it complete and should be processed after active files."), *Configuration.ConstructList[IndexToConstruct]);
+										// Only start the file if there's not currently an active file - otherwise we
+										// do it when the current one finishes.
+										if (ActiveFileCount == 0)
+										{
+											FileConstructorStat->OnFileStarted(AnotherFile->BuildFilename, AnotherFile->FileManifest->FileSize);
+										}
+
+										FileToStart = AnotherFile.Get();
+										ActiveFiles.Push(MoveTemp(AnotherFile));
 									}
 								} // end if not delaying for dependencies
 
@@ -4281,7 +4095,7 @@ void FBuildPatchFileConstructor::ConstructFiles(FResumeData& ResumeData)
 		{
 			if (!File->bSuccess && File->ConstructionError != EConstructionError::Aborted)
 			{
-				CompleteConstructedFile(*File.Get(), ResumeData);
+				CompleteConstructedFile(*File.Get());
 			}
 		}
 
@@ -4298,17 +4112,10 @@ void FBuildPatchFileConstructor::ConstructFiles(FResumeData& ResumeData)
 				File->bSuccess = false;
 				File->ConstructionError = EConstructionError::Aborted;
 			}
-			CompleteConstructedFile(*File.Get(), ResumeData);
+			CompleteConstructedFile(*File.Get());
 		}
 	}
 
-	if (CorruptFiles.Num() > 0)
-	{
-		Abort();
-		UE_LOG(LogBuildPatchServices, Error, TEXT("FBuildPatchFileConstructor: The %d files are corrupt."), CorruptFiles.Num());
-		InstallerError->SetError(EBuildPatchInstallError::FileConstructionFail, ConstructionErrorCodes::OutboundCorrupt);
-	}
-
 	uint64 ConstructCycles = FPlatformTime::Cycles64() - ConstructStartCycles;
 	uint64 UnaccountedForCycles = ConstructCycles - HashCycles - WaitCycles - ReadCheckCycles - ReadDoneCycles - WriteStartCycles - WriteDoneCycles - CloudTickCycles;
 
@@ -4325,31 +4132,3 @@ void FBuildPatchFileConstructor::ConstructFiles(FResumeData& ResumeData)
 		100.0 * UnaccountedForCycles / ConstructCycles
 		);
 }
-
-void FBuildPatchFileConstructor::RemoveUnnecessaryStagedFiles(const FString& StagingDirectory) const
-{
-	// Read staging folder to get the list of existing staged files
-	TSet<FString> AllStagedFiles;
-	IFileManager::Get().IterateDirectory(*StagingDirectory, [&AllStagedFiles, StagingDirectory](const TCHAR* Filename, bool) {
-			AllStagedFiles.Add(FString(Filename).Mid(StagingDirectory.Len() + 1)); // get relative path.
-			return true;
-		});
-
-	// Get a set of potentially staged files based on ConstructionList
-	TSet<FString> NecessaryStagedFiles;
-	Algo::Transform(ConstructionList, NecessaryStagedFiles, [](const FFileToConstruct& FileToConstruct)
-		{
-			return FileToConstruct.StagingFilename;
-		});
-
-	const TSet<FString> UnnecessaryFiles = AllStagedFiles.Difference(NecessaryStagedFiles);
-	for (const FString& Filename : UnnecessaryFiles)
-	{
-		UE_LOG(LogBuildPatchFileConstructor, Log, TEXT("RemoveUnnecessaryStagedFiles: Deleting unnecessary staged file %s"), *Filename);
-		if (!FileSystem->DeleteFile(*(StagingDirectory / Filename)))
-		{
-			UE_LOG(LogBuildPatchServices, Error, TEXT("RemoveUnnecessaryStagedFiles: Error deleting file: %s (Error Code %i)"), *Filename, FPlatformMisc::GetLastError());
-		}
-	}
-}
-
diff --git a/Engine/Source/Runtime/Online/BuildPatchServices/Private/BuildPatchFileConstructor.h b/Engine/Source/Runtime/Online/BuildPatchServices/Private/BuildPatchFileConstructor.h
--- a/Engine/Source/Runtime/Online/BuildPatchServices/Private/BuildPatchFileConstructor.h
+++ b/Engine/Source/Runtime/Online/BuildPatchServices/Private/BuildPatchFileConstructor.h
@@ -12,7 +12,6 @@
 #include "Installer/Controllable.h"
 #include "Common/SpeedRecorder.h"
 #include "BuildPatchInstall.h"
-#include "Async/Future.h"
 
 
 // Forward declarations
@@ -21,7 +20,6 @@ struct FResumeData;
 struct FFileConstructionState;
 class FBuildPatchAppManifest;
 enum class EConstructionError : uint8;
-enum class EFileStage : uint8;
 
 class IBuildInstallerSharedContext;
 class FChunkBackingStore;
@@ -76,11 +74,6 @@ namespace BuildPatchServices
 
 		// The install mode used for this installation.
 		EInstallMode InstallMode;
-		// Whether we are needing to use SHA1 hash string value as filename for temp staging files.
-		bool bUseSHA1StageFilenames;
-
-		// If true, fail the installation if the manifest contains a symlink.
-		bool bRejectSymlinks = false;
 
 		// The location where memory overflow will get written to.
 		FString BackingStoreDirectory;
@@ -143,16 +136,10 @@ namespace BuildPatchServices
 		// IControllable interface begin.
 		virtual void SetPaused(bool bInIsPaused) override;
 		virtual void Abort() override;
-
-		virtual void Reset() override
-		{
-			// This should _never_ be called. I'm not sure how you can Reset() a constructor..
-			check(false);
-		}
 		// IControllable interface end.
 
 		void WakeUpDispatch();
-		
+
 		/**
 		 * Get the disk space that was required to perform the installation. This can change over time and indicates the required
 		 * space to _finish_ the installation from the current state. It is not initialized until after resume is processed and returns
@@ -170,26 +157,13 @@ namespace BuildPatchServices
 		 */
 		uint64 GetAvailableDiskSpace();
 
-		/**
-		 * Gets flag indicating that at least one file is corrupt due to invalid hash sum.
-		 * @return	true if at least one file is corrupt.
-		 */
-		bool HasCorruptFiles() const;
-
 		/**
 		 * Broadcasts with full filepath to file that the constructor is about to delete in order to free up space.
 		 * @return	Reference to the event object.
 		 */
-		DECLARE_EVENT_OneParam(FBuildPatchFileConstructor, FOnBeforeDeleteFile, const FString& /*FullFilePath*/);
+		DECLARE_EVENT_OneParam(FBuildPatchFileConstructor, FOnBeforeDeleteFile, const FString& /*BuildFile*/);
 		FOnBeforeDeleteFile& OnBeforeDeleteFile();
 
-		/**
-		 * Broadcasts with the build relative path for a file that has finished being staged (including by a previous incomplete installation).
-		 * @return	Reference to the event object.
-		 */
-		DECLARE_EVENT_OneParam(FBuildPatchFileConstructor, FOnFileIsStaged, const FString& /*BuildFilename*/);
-		FOnFileIsStaged& OnFileIsStaged();
-
 		void GrabFilesInstalledToMemory(TMap<FString, TArray64<uint8>>& OutFilesInstalledToMemory)
 		{
 			OutFilesInstalledToMemory = MoveTemp(MemoryOutputFiles);
@@ -218,32 +192,19 @@ namespace BuildPatchServices
 		 */
 		uint64 CalculateInProgressDiskSpaceRequired(const FFileManifest& InProgressFileManifest, uint64 InProgressFileSize);
 
-		/**
-		 * Calculates the amount of disk space we need to finish the install, needs to be called on file boundaries on the construct thread.
-		 * @param InProgressFileSize		The remaining size required for the file currently being constructed.
-		 * @return Number of bytes needed to complete the installation, or zero if space will be freed by removing chunks.
-		 */
-		uint64 CalculateDiskSpaceRequirementsWithDeleteDuringInstall(int64 InProgressFileSize);
+		// Calculates the amount of disk space we need to finish the install, needs to be called on file boundaries on the construct thread.
+		uint64 CalculateDiskSpaceRequirementsWithDeleteDuringInstall();
 
 		/**
 		 * Sequentially build the files required, potentially skipping already created files or resuming partial
 		 * files.
 		 */
-		void ConstructFiles(FResumeData& ResumeData);
+		void ConstructFiles(const FResumeData& ResumeData);
 
 	private:
 		// The configuration for the constructor.
 		const FFileConstructorConfig Configuration;
 
-		// A flag marking that we a running.
-		bool bIsRunning;
-
-		// A flag marking that we initialized correctly.
-		bool bIsInited;
-
-		// A flag marking that our init returned a failure (true means failed).
-		bool bInitFailed;
-
 		// A flag marking that we told the chunk cache to queue required downloads.
 		bool bIsDownloadStarted;
 
@@ -275,8 +236,6 @@ namespace BuildPatchServices
 			// constructed in order, we only track the file that will be harvested last. If this
 			// file has no dependencies this is -1
 			int32 LatestDependentInstallSource = -1;
-			// Name of the staging file. It may be obfuscated if bUseSHA1StageFilenames is true.
-			FString StagingFilename;
 		};
 
 		// The in-oder list of files to construct. The array is parallel with Configuration.ConstructList.
@@ -304,8 +263,6 @@ namespace BuildPatchServices
 		// Pointer to the chunk reference tracker.
 		IChunkReferenceTracker* ChunkReferenceTracker;
 
-		std::atomic<int32> RemainingChunkCount = TNumericLimits<int32>::Max();
-
 		// Pointer to the installer error class.
 		IInstallerError* InstallerError;
 
@@ -343,9 +300,7 @@ namespace BuildPatchServices
 
 		// Event executed before deleting an old installation file.
 		FOnBeforeDeleteFile BeforeDeleteFileEvent;
-		// Event executed when a file has now been staged or was previously staged.
-		FOnFileIsStaged FileIsStagedEvent;
-		
+
 
 		TArray<FEvent*> ThreadWakeups;
 		TArray<TArray<IConstructorChunkSource::FRequestProcessFn>> ThreadJobPostings;
@@ -379,19 +334,10 @@ namespace BuildPatchServices
 		void OpenFileToConstruct(FFileConstructionState& CurrentFile);
 		bool HandleInitialDiskSizeCheck(const FFileManifest& CurrentFileManifest, int64 BytesInToCurrentFile);
 		bool HarvestChunksForCompletedFile(const FString& InstallationDirectory, const FString& FilenameInManifest);
-		void CompleteConstructedFile(FFileConstructionState& CurrentFile, FResumeData& ResumeData);
-
-		// Determine and return the file stage based on its existence and file size.
-		EFileStage GetFileStage(const FFileConstructionState& CurrentFile, const FResumeData& ResumeData) const;
-
-		// Initialize and validate.
+		void CompleteConstructedFile(FFileConstructionState& CurrentFile);
 		void InitFile(FFileConstructionState& CurrentFile, const FResumeData& ResumeData);
-		
-		// Removes staged files that are no longer required for the current operation.
-		void RemoveUnnecessaryStagedFiles(const FString& StagingDirectory) const;
 
-		// Set of files considered corrupt due to an invalid hash sum.
-		TSet<FString> CorruptFiles;
+
 
 	public:
 		struct FBackingStoreStats
@@ -518,5 +464,4 @@ namespace BuildPatchServices
 namespace FileConstructorHelpers
 {
 	uint64 CalculateRequiredDiskSpace(const FBuildPatchAppManifestPtr& CurrentManifest, const FBuildPatchAppManifestRef& BuildManifest, const BuildPatchServices::EInstallMode& InstallMode, const TSet<FString>& InstallTags);
-	bool CheckRemainingDiskSpace(const FString& InstallDirectory, uint64 RemainingBytesRequired, uint64& OutAvailableDiskSpace);
 }
