[HTTPServer] Allow response/callback to be used for multiple writes (sse text event streams)

#Synced-CL 51056504

#Synced-CL 51056504

diff --git a/Engine/Source/Runtime/Online/HTTPServer/Private/HttpConnection.cpp b/Engine/Source/Runtime/Online/HTTPServer/Private/HttpConnection.cpp
--- a/Engine/Source/Runtime/Online/HTTPServer/Private/HttpConnection.cpp
+++ b/Engine/Source/Runtime/Online/HTTPServer/Private/HttpConnection.cpp
@@ -179,7 +179,7 @@ void FHttpConnection::CompleteRead(const TSharedPtr<FHttpServerRequest>& Request
 				TEXT("EndProcessRequest\t [%d][%u]-%u : %u"), 
 						 OriginPortCapture, ConnectionIdCapture, LastRequestNumberCapture, int(Response->Code));
 
-			// Ensure this result callback was called once
+			// Ensure this result callback was called at least once
 			check(EHttpConnectionState::AwaitingProcessing == SharedThisPtr->GetState());
 
 			// Begin response flow
@@ -273,7 +273,14 @@ void FHttpConnection::CompleteWrite()
 
 	if (bKeepAlive && !bGracefulDestroyRequested)
 	{
-		ChangeState(EHttpConnectionState::AwaitingRead);
+		if (WriteContext.HasAdditionalWrites())
+		{
+			ChangeState(EHttpConnectionState::AwaitingProcessing);
+		}
+		else
+		{
+			ChangeState(EHttpConnectionState::AwaitingRead);
+		}
 	}
 	else
 	{
diff --git a/Engine/Source/Runtime/Online/HTTPServer/Private/HttpConnectionResponseWriteContext.cpp b/Engine/Source/Runtime/Online/HTTPServer/Private/HttpConnectionResponseWriteContext.cpp
--- a/Engine/Source/Runtime/Online/HTTPServer/Private/HttpConnectionResponseWriteContext.cpp
+++ b/Engine/Source/Runtime/Online/HTTPServer/Private/HttpConnectionResponseWriteContext.cpp
@@ -26,12 +26,22 @@ void FHttpConnectionResponseWriteContext::ResetContext(TUniquePtr<FHttpServerRes
 
 	if (Response)
 	{
-		// Affix Content-Length Header
-		TArray<FString> ContentLengthValue = { FString::FromInt(Response->Body.Num()) };
-		Response->Headers.Add(UE_HTTP_SERVER_HEADER_KEYS_CONTENT_LENGTH, MoveTemp(ContentLengthValue));
+		const bool bResponseIsStream = EnumHasAnyFlags(Response->Flags, EHttpServerResponseFlags::MultipleWriteStream);
 
-		// Serialize Headers
-		HeaderBytes.Append(SerializeHeadersUtf8(Response->HttpVersion, Response->Code, Response->Headers));
+		if (!bResponseIsStream)
+		{
+			// Affix Content-Length Header
+			TArray<FString> ContentLengthValue = { FString::FromInt(Response->Body.Num()) };
+			Response->Headers.Add(UE_HTTP_SERVER_HEADER_KEYS_CONTENT_LENGTH, MoveTemp(ContentLengthValue));
+		}
+
+		const bool bWriteHeaders = !EnumHasAnyFlags(Response->Flags, EHttpServerResponseFlags::SkipHeaderWrite);
+
+		if (bWriteHeaders)
+		{
+			// Serialize Headers
+			HeaderBytes.Append(SerializeHeadersUtf8(Response->HttpVersion, Response->Code, Response->Headers));
+		}
 	}
 }
 
diff --git a/Engine/Source/Runtime/Online/HTTPServer/Private/HttpConnectionResponseWriteContext.h b/Engine/Source/Runtime/Online/HTTPServer/Private/HttpConnectionResponseWriteContext.h
--- a/Engine/Source/Runtime/Online/HTTPServer/Private/HttpConnectionResponseWriteContext.h
+++ b/Engine/Source/Runtime/Online/HTTPServer/Private/HttpConnectionResponseWriteContext.h
@@ -37,6 +37,22 @@ public:
 	 */
 	void ResetContext(TUniquePtr<FHttpServerResponse>&&  Response);
 
+	/**
+	 * Determine if the connection should expect multiple writes as part of this response before the next 
+	 * read/close by re-using an existing FHttpResultCallback.  Streams should issue the first response callback 
+	 * with an empty body to write headers.
+	 *
+	 * @return true if there are multiple writes expected
+	 */
+	bool IsMultipleWriteStream() const { return Response.IsValid() && EnumHasAnyFlags(Response->Flags, EHttpServerResponseFlags::MultipleWriteStream); }
+
+	/**
+	 * Determine if the context's current state is expecting additional body writes
+	 *
+	 * @return true if there are more writes
+	 */
+	bool HasAdditionalWrites() const { return Response.IsValid() && EnumHasAllFlags(Response->Flags, EHttpServerResponseFlags::MultipleWriteStream | EHttpServerResponseFlags::HasAdditionalWrites); }
+
 private:
 
 	/**
diff --git a/Engine/Source/Runtime/Online/HTTPServer/Public/HttpServerResponse.h b/Engine/Source/Runtime/Online/HTTPServer/Public/HttpServerResponse.h
--- a/Engine/Source/Runtime/Online/HTTPServer/Public/HttpServerResponse.h
+++ b/Engine/Source/Runtime/Online/HTTPServer/Public/HttpServerResponse.h
@@ -13,6 +13,22 @@
 
 class FUtf8String;
 
+enum class EHttpServerResponseFlags : uint8
+{
+	None = 0,
+
+	/** Used to indicate the response body will consist of multiple writes before another read/close (possibly an SSE stream) */
+	MultipleWriteStream = 1 << 0,
+
+	/** Used to indicate if additional stream writes are expected, used with MultipleWriteStream */
+	HasAdditionalWrites = 1 << 1,
+
+	/** Used to indicate if the next response write should skip including headers, used with MultipleWriteStream */
+	SkipHeaderWrite		= 1 << 2,
+};
+
+ENUM_CLASS_FLAGS(EHttpServerResponseFlags);
+
 struct FHttpServerResponse final
 {
 public:
@@ -33,10 +49,10 @@ public:
 	{ }
 
 	/** Http protocol version */
-	HttpVersion::EHttpServerHttpVersion HttpVersion;
+	HttpVersion::EHttpServerHttpVersion HttpVersion = HttpVersion::EHttpServerHttpVersion::HTTP_VERSION_UNKNOWN;
 
 	/** Http Response Code */
-	EHttpServerResponseCodes Code;
+	EHttpServerResponseCodes Code = EHttpServerResponseCodes::Unknown;
 
 	/** Http Headers */
 	TMap<FString, TArray<FString>> Headers;
@@ -44,6 +60,9 @@ public:
 	/** Http Body Content */
 	TArray<uint8> Body;
 
+	/** Response flags */
+	EHttpServerResponseFlags Flags = EHttpServerResponseFlags::None;
+
 public:
 
 	/**
