[BuildPatchServices] More Mac/Unix fallout fixes, include unistd.h

diff --git a/Engine/Source/Runtime/Online/BuildPatchServices/Private/BuildPatchInstaller.cpp b/Engine/Source/Runtime/Online/BuildPatchServices/Private/BuildPatchInstaller.cpp
--- a/Engine/Source/Runtime/Online/BuildPatchServices/Private/BuildPatchInstaller.cpp
+++ b/Engine/Source/Runtime/Online/BuildPatchServices/Private/BuildPatchInstaller.cpp
@@ -21,6 +21,10 @@
 #include "Misc/Paths.h"
 #include "Misc/ScopeLock.h"
 
+#if PLATFORM_MAC || PLATFORM_LINUX
+#include <unistd.h>
+#endif
+
 #include "BuildPatchFileConstructor.h"
 #include "BuildPatchServicesModule.h"
 #include "BuildPatchServicesPrivate.h"
