Fix FHttpManager name collision in BuildPatchServices
#rnx

#Synced-CL 51101942

diff --git a/Engine/Source/Runtime/Online/BuildPatchServices/Private/Common/HttpManager.cpp b/Engine/Source/Runtime/Online/BuildPatchServices/Private/Common/HttpManager.cpp
--- a/Engine/Source/Runtime/Online/BuildPatchServices/Private/Common/HttpManager.cpp
+++ b/Engine/Source/Runtime/Online/BuildPatchServices/Private/Common/HttpManager.cpp
@@ -5,12 +5,12 @@
 
 namespace BuildPatchServices
 {
-	class FHttpManager
+	class FBPSHttpManager
 		: public IHttpManager
 	{
 	public:
-		FHttpManager();
-		~FHttpManager();
+		FBPSHttpManager();
+		~FBPSHttpManager();
 
 		// IHttpManager interface begin.
 		virtual TSharedRef<IHttpRequest, ESPMode::ThreadSafe> CreateRequest() override;
@@ -20,22 +20,22 @@ namespace BuildPatchServices
 		FHttpModule& HttpModule;
 	};
 
-	FHttpManager::FHttpManager()
+	FBPSHttpManager::FBPSHttpManager()
 		: HttpModule(FHttpModule::Get())
 	{
 	}
 
-	FHttpManager::~FHttpManager()
+	FBPSHttpManager::~FBPSHttpManager()
 	{
 	}
 
-	TSharedRef<IHttpRequest, ESPMode::ThreadSafe> FHttpManager::CreateRequest()
+	TSharedRef<IHttpRequest, ESPMode::ThreadSafe> FBPSHttpManager::CreateRequest()
 	{
 		return HttpModule.CreateRequest();
 	}
 
 	IHttpManager* FHttpManagerFactory::Create()
 	{
-		return new FHttpManager();
+		return new FBPSHttpManager();
 	}
 }
\ No newline at end of file
