Added caching of Ini Config values inside HttpServerConfig to avoid a 30ms spike on the server when creating a new HttpConnection

[RN] Improved ini value caching of HttpServerConfig to avoid the cost of Config SafeReload
#tests Built and tested on Dedicated Server with Cooked Data and Editor
#Synced-CL 51046133

diff --git a/Engine/Source/Runtime/Online/HTTPServer/Private/HttpServerConfig.cpp b/Engine/Source/Runtime/Online/HTTPServer/Private/HttpServerConfig.cpp
--- a/Engine/Source/Runtime/Online/HTTPServer/Private/HttpServerConfig.cpp
+++ b/Engine/Source/Runtime/Online/HTTPServer/Private/HttpServerConfig.cpp
@@ -3,53 +3,92 @@
 #include "HttpServerConfig.h"
 #include "CoreMinimal.h"
 #include "Misc/ConfigCacheIni.h"
+#include "Misc/CoreDelegates.h"
 
 DEFINE_LOG_CATEGORY(LogHttpServerConfig)
 
+namespace FHttpServerConfigCache
+{
+	static FDelegateHandle ConfigSectionChangedHandle;
+	static bool bConnectionsCacheDirty = true;
+	static bool bListenersCacheDirty = true;
+	static FHttpServerConnectionConfig CacheConnectionConfig;
+	static FHttpServerListenerConfig CacheListenerConfig;
+	static TArray<FString> CacheListenerConfigs;
+	static const FString IniSectionNameHTTPServerConnections("HTTPServer.Connections");
+	static const FString IniSectionNameHTTPServerListeners("HTTPServer.Listeners");
+}
+
+void FHttpServerConfig::OnConfigSectionsChanged(const FString& IniFilename, const TSet<FString>& SectionNames)
+{
+	using namespace FHttpServerConfigCache;
+
+	if (IniFilename == GEngineIni)
+	{
+		if (SectionNames.Contains(IniSectionNameHTTPServerConnections))
+		{
+			bConnectionsCacheDirty = true;
+		}
+		if (SectionNames.Contains(IniSectionNameHTTPServerListeners))
+		{
+			bListenersCacheDirty = true;
+		}
+	}
+}
+
 const FHttpServerListenerConfig FHttpServerConfig::GetListenerConfig(uint32 Port) 
 {
-	static const FString IniSectionName("HTTPServer.Listeners");
+	using namespace FHttpServerConfigCache;
 
-	// Code default values
-	FHttpServerListenerConfig Config;
+	if (!ConfigSectionChangedHandle.IsValid())
+	{
+		ConfigSectionChangedHandle = FCoreDelegates::TSOnConfigSectionsChanged().AddStatic(FHttpServerConfig::OnConfigSectionsChanged);
+	}
 
-	// Apply default ini configuration
-	GConfig->GetString(*IniSectionName, TEXT("DefaultBindAddress"), Config.BindAddress, GEngineIni);
-	GConfig->GetInt(*IniSectionName, TEXT("DefaultBufferSize"), Config.BufferSize, GEngineIni);
-	GConfig->GetInt(*IniSectionName, TEXT("DefaultConnectionsBacklogSize"), Config.ConnectionsBacklogSize, GEngineIni);
-	GConfig->GetInt(*IniSectionName, TEXT("DefaultMaxConnectionsAcceptPerFrame"), Config.MaxConnectionsAcceptPerFrame, GEngineIni);
-	GConfig->GetBool(*IniSectionName, TEXT("DefaultReuseAddressAndPort"), Config.bReuseAddressAndPort, GEngineIni);
+	// Update values from Config if necessary
+	if (bListenersCacheDirty)
+	{
+		// Apply default ini configuration
+		GConfig->GetString(*IniSectionNameHTTPServerListeners, TEXT("DefaultBindAddress"), CacheListenerConfig.BindAddress, GEngineIni);
+		GConfig->GetInt(*IniSectionNameHTTPServerListeners, TEXT("DefaultBufferSize"), CacheListenerConfig.BufferSize, GEngineIni);
+		GConfig->GetInt(*IniSectionNameHTTPServerListeners, TEXT("DefaultConnectionsBacklogSize"), CacheListenerConfig.ConnectionsBacklogSize, GEngineIni);
+		GConfig->GetInt(*IniSectionNameHTTPServerListeners, TEXT("DefaultMaxConnectionsAcceptPerFrame"), CacheListenerConfig.MaxConnectionsAcceptPerFrame, GEngineIni);
+		GConfig->GetBool(*IniSectionNameHTTPServerListeners, TEXT("DefaultReuseAddressAndPort"), CacheListenerConfig.bReuseAddressAndPort, GEngineIni);
+
+		GConfig->GetArray(*IniSectionNameHTTPServerListeners, TEXT("ListenerOverrides"), CacheListenerConfigs, GEngineIni);
+
+		bListenersCacheDirty = false;
+	}
+	
+	// Setup default values
+	FHttpServerListenerConfig Config(CacheListenerConfig);
 
 	// Apply per-port ini overrides
-	TArray<FString> ListenerConfigs;
-	if (GConfig->GetArray(*IniSectionName, TEXT("ListenerOverrides"), ListenerConfigs, GEngineIni))
+	for(FString ListenerConfigStr : CacheListenerConfigs)
 	{
-		for (FString ListenerConfigStr : ListenerConfigs)
+		ListenerConfigStr.TrimStartAndEndInline();
+		ListenerConfigStr.ReplaceInline(TEXT("("), TEXT(""));
+		ListenerConfigStr.ReplaceInline(TEXT(")"), TEXT(""));
+
+		// Listener config overrides must specify a port
+		uint32 ConfiguredPort = 0;
+		if (!FParse::Value(*ListenerConfigStr, TEXT("Port="), ConfiguredPort))
+		{
+			UE_LOG(LogHttpServerConfig, Error,
+				TEXT("ListenerOverride: %s does not specify required Port parameter"),
+				*ListenerConfigStr);
+			continue;
+		}
+
+		if (Port == ConfiguredPort)
 		{
-			ListenerConfigStr.TrimStartAndEndInline();
-			ListenerConfigStr.ReplaceInline(TEXT("("), TEXT(""));
-			ListenerConfigStr.ReplaceInline(TEXT(")"), TEXT(""));
-
-			// Listener config overrides must specify a port
-			uint32 ConfiguredPort = 0;
-			if (!FParse::Value(*ListenerConfigStr, TEXT("Port="), ConfiguredPort))
-			{
-				UE_LOG(LogHttpServerConfig, Error,
-					TEXT("ListenerOverride: %s does not specify required Port parameter"),
-					*ListenerConfigStr);
-				continue;
-			}
-
-			if (Port == ConfiguredPort)
-			{
-				// override defaults with config values
-				FParse::Value(*ListenerConfigStr, TEXT("BindAddress="), Config.BindAddress);
-				FParse::Value(*ListenerConfigStr, TEXT("BufferSize="), Config.BufferSize);
-				FParse::Value(*ListenerConfigStr, TEXT("ConnectionsBacklogSize="), Config.ConnectionsBacklogSize);
-				FParse::Value(*ListenerConfigStr, TEXT("MaxConnectionsAcceptPerFrame="), Config.MaxConnectionsAcceptPerFrame);
-				FParse::Bool(*ListenerConfigStr, TEXT("ReuseAddressAndPort="), Config.bReuseAddressAndPort);
-				break;
-			}
+			// override defaults with config values
+			FParse::Value(*ListenerConfigStr, TEXT("BindAddress="), Config.BindAddress);
+			FParse::Value(*ListenerConfigStr, TEXT("BufferSize="), Config.BufferSize);
+			FParse::Value(*ListenerConfigStr, TEXT("ConnectionsBacklogSize="), Config.ConnectionsBacklogSize);
+			FParse::Value(*ListenerConfigStr, TEXT("MaxConnectionsAcceptPerFrame="), Config.MaxConnectionsAcceptPerFrame);
+			FParse::Bool(*ListenerConfigStr, TEXT("ReuseAddressAndPort="), Config.bReuseAddressAndPort);
+			break;
 		}
 	}
 	return Config;
@@ -57,13 +96,19 @@ const FHttpServerListenerConfig FHttpServerConfig::GetListenerConfig(uint32 Port
 
 const FHttpServerConnectionConfig FHttpServerConfig::GetConnectionConfig()
 {
-	static const FString IniSectionName("HTTPServer.Connections");
+	using namespace FHttpServerConfigCache;
 
-	// Code default values
-	FHttpServerConnectionConfig Config;
+	if (!ConfigSectionChangedHandle.IsValid())
+	{
+		ConfigSectionChangedHandle = FCoreDelegates::TSOnConfigSectionsChanged().AddStatic(FHttpServerConfig::OnConfigSectionsChanged);
+	}
 
-	// Apply default ini configuration
-	GConfig->GetFloat(*IniSectionName, TEXT("BeginReadWaitTimeMS"), Config.BeginReadWaitTimeMS, GEngineIni);
+	if (bConnectionsCacheDirty)
+	{
+		// Apply default ini configuration
+		GConfig->GetFloat(*IniSectionNameHTTPServerConnections, TEXT("BeginReadWaitTimeMS"), CacheConnectionConfig.BeginReadWaitTimeMS, GEngineIni);
+		bConnectionsCacheDirty = false;
+	}
 
-	return Config;
-}
+	return CacheConnectionConfig;
+}
\ No newline at end of file
diff --git a/Engine/Source/Runtime/Online/HTTPServer/Private/HttpServerConfig.h b/Engine/Source/Runtime/Online/HTTPServer/Private/HttpServerConfig.h
--- a/Engine/Source/Runtime/Online/HTTPServer/Private/HttpServerConfig.h
+++ b/Engine/Source/Runtime/Online/HTTPServer/Private/HttpServerConfig.h
@@ -43,4 +43,9 @@ struct FHttpServerConfig
 	* @return The configuration if configured, or default configuration
 	*/
 	static const FHttpServerConnectionConfig GetConnectionConfig();
+
+	/**
+	* Callback when the IniConfig is updated so that the internal cache can be updated
+	*/
+	static void OnConfigSectionsChanged(const FString& IniFilename, const TSet<FString>& SectionNames);
 };
\ No newline at end of file
