BPS: Fix for crash during install in some cases.

#rnx
#Synced-CL 51285544

diff --git a/Engine/Source/Runtime/Online/BuildPatchServices/Private/BuildPatchManifestSet.cpp b/Engine/Source/Runtime/Online/BuildPatchServices/Private/BuildPatchManifestSet.cpp
--- a/Engine/Source/Runtime/Online/BuildPatchServices/Private/BuildPatchManifestSet.cpp
+++ b/Engine/Source/Runtime/Online/BuildPatchServices/Private/BuildPatchManifestSet.cpp
@@ -168,15 +168,10 @@ namespace BuildPatchServices
 
 		FDataReference const* const DataLookupFind(const FGuid& DataGuid) const
 		{
-			static thread_local FDataReference const* LatestLookupResult = nullptr;
-			static thread_local FGuid LatestGuid = FGuid();
-			if (LatestGuid == DataGuid && LatestLookupResult != nullptr)
-			{
-				return LatestLookupResult;
-			}
-			LatestGuid = DataGuid;
-			LatestLookupResult = DataLookup.Find(DataGuid);
-			return LatestLookupResult;
+			// This used to be a thread_local static lookup cache however it fails with multiple installations as well
+			// as if we have multiple installs in a row that happen to use the same guid. I'm not confident using a thread lock
+			// is worth saving what should be a N(1) lookup.
+			return DataLookup.Find(DataGuid);
 		}
 
 		virtual uint64 GetDownloadSize(const FGuid& DataGuid) const override
