[Http]Fix CancelRequest call not working during WillRetry delegate or during back off.

#Synced-CL 51366129

diff --git a/Engine/Source/Programs/WebTests/Private/TestHttp.cpp b/Engine/Source/Programs/WebTests/Private/TestHttp.cpp
--- a/Engine/Source/Programs/WebTests/Private/TestHttp.cpp
+++ b/Engine/Source/Programs/WebTests/Private/TestHttp.cpp
@@ -2632,6 +2632,102 @@ TEST_CASE_METHOD(FWaitUntilCompleteHttpFixture, "Request can time out during ret
 	HttpRequest->ProcessRequest();
 }
 
+TEST_CASE_METHOD(FWaitUntilCompleteHttpFixture, "Request can be canceled during WillRetry delegate", HTTP_TAG)
+{
+	if (!bRetryEnabled)
+	{
+		return;
+	}
+
+	DisableWarningsInThisTest();
+
+	TSharedRef<IHttpRequest> HttpRequest = HttpRetryManager->CreateRequest(
+		1/*InRetryLimitCountOverride*/,
+		FHttpRetrySystem::FRetryTimeoutRelativeSecondsSetting()/*InRetryTimeoutRelativeSecondsOverride unused*/,
+		{EHttpResponseCodes::TooManyRequests}/*InRetryResponseCodes*/
+	);
+
+	HttpRequest->SetURL(UrlMockStatus(EHttpResponseCodes::TooManyRequests));
+
+	uint32 RetryAfter = 2;
+	// Will be forwarded back in response
+	HttpRequest->SetHeader(TEXT("Retry-After"), FString::Format(TEXT("{0}"), { RetryAfter }));
+
+	++ExpectingExtraCallbacks;
+	HttpRequest->OnRequestWillRetry().BindLambda([this](FHttpRequestPtr Request, FHttpResponsePtr /*Response*/, float LockoutPeriod) {
+		--ExpectingExtraCallbacks;
+		Request->CancelRequest();
+	});
+
+	const double StartTime = FPlatformTime::Seconds();
+	HttpRequest->OnProcessRequestComplete().BindLambda([StartTime](FHttpRequestPtr HttpRequest, FHttpResponsePtr HttpResponse, bool bSucceeded) {
+		CHECK(HttpRequest->GetStatus() == EHttpRequestStatus::Failed);
+		CHECK(HttpRequest->GetFailureReason() == EHttpFailureReason::Cancelled);
+		CHECK(!bSucceeded);
+		CHECK(HttpResponse == nullptr);
+
+		// Will not wait for 2s(RetryAfter) if canceled during WillRetry delegate
+		const double DurationInSeconds  = FPlatformTime::Seconds() - StartTime;
+		CHECK(FMath::IsNearlyEqual(DurationInSeconds, 0.0, HTTP_TIME_DIFF_TOLERANCE_OF_REQUEST));
+	});
+
+	HttpRequest->ProcessRequest();
+}
+
+TEST_CASE_METHOD(FWaitUntilCompleteHttpFixture, "Request can be canceled during back off", HTTP_TAG)
+{
+	if (!bRetryEnabled)
+	{
+		return;
+	}
+
+	DisableWarningsInThisTest();
+
+	TSharedRef<IHttpRequest> HttpRequest = HttpRetryManager->CreateRequest(
+		1/*InRetryLimitCountOverride*/,
+		FHttpRetrySystem::FRetryTimeoutRelativeSecondsSetting()/*InRetryTimeoutRelativeSecondsOverride unused*/,
+		{EHttpResponseCodes::TooManyRequests}/*InRetryResponseCodes*/
+	);
+
+	HttpRequest->SetURL(UrlMockStatus(EHttpResponseCodes::TooManyRequests));
+
+	uint32 RetryAfter = 4;
+	// Will be forwarded back in response
+	HttpRequest->SetHeader(TEXT("Retry-After"), FString::Format(TEXT("{0}"), { RetryAfter }));
+
+	++ExpectingExtraCallbacks;
+	HttpRequest->OnRequestWillRetry().BindLambda([this](FHttpRequestPtr Request, FHttpResponsePtr /*Response*/, float LockoutPeriod) {
+		--ExpectingExtraCallbacks;
+	});
+
+	uint32 CancelAfter = 2;
+
+	const double StartTime = FPlatformTime::Seconds();
+	HttpRequest->OnProcessRequestComplete().BindLambda([StartTime, CancelAfter](FHttpRequestPtr HttpRequest, FHttpResponsePtr HttpResponse, bool bSucceeded) {
+		CHECK(HttpRequest->GetStatus() == EHttpRequestStatus::Failed);
+		CHECK(HttpRequest->GetFailureReason() == EHttpFailureReason::Cancelled);
+		CHECK(!bSucceeded);
+		CHECK(HttpResponse == nullptr);
+
+		// Will not wait for total 4s(RetryAfter) if canceled during back off
+		const double DurationInSeconds  = FPlatformTime::Seconds() - StartTime;
+		CHECK(FMath::IsNearlyEqual(DurationInSeconds, CancelAfter, HTTP_TIME_DIFF_TOLERANCE_OF_REQUEST));
+	});
+
+	HttpRequest->ProcessRequest();
+
+	while (HasOngoingRequest())
+	{
+		TickHttpManager();
+
+		if (FPlatformTime::Seconds() - StartTime > CancelAfter)
+		{
+			HttpRequest->CancelRequest();
+			break;
+		}
+	}
+}
+
 TEST_CASE_METHOD(FWaitUntilCompleteHttpFixture, "Request will not retry", HTTP_TAG)
 {
 	if (!bRetryEnabled)
diff --git a/Engine/Source/Runtime/Online/HTTP/Private/Curl/CurlHttp.cpp b/Engine/Source/Runtime/Online/HTTP/Private/Curl/CurlHttp.cpp
--- a/Engine/Source/Runtime/Online/HTTP/Private/Curl/CurlHttp.cpp
+++ b/Engine/Source/Runtime/Online/HTTP/Private/Curl/CurlHttp.cpp
@@ -1101,6 +1101,7 @@ void FCurlHttpRequest::ClearInCaseOfRetry()
 	LastReportedBytesRead = 0;
 	TotalBytesRead = 0;
 	bAnyHttpActivity = false;
+	bCurlRequestCompleted = false;
 
 	// Clear the info cache log so we don't output messages from previous requests when reusing/retrying a request
 	{
diff --git a/Engine/Source/Runtime/Online/HTTP/Private/HttpRetrySystem.cpp b/Engine/Source/Runtime/Online/HTTP/Private/HttpRetrySystem.cpp
--- a/Engine/Source/Runtime/Online/HTTP/Private/HttpRetrySystem.cpp
+++ b/Engine/Source/Runtime/Online/HTTP/Private/HttpRetrySystem.cpp
@@ -283,15 +283,25 @@ void FHttpRetrySystem::FRequest::HttpOnProcessRequestComplete(FHttpRequestPtr In
 				const double NowAbsoluteSeconds = FPlatformTime::Seconds();
 				float LockoutPeriod = RetryManagerPtr->GetLockoutPeriodSeconds(HttpRetryRequestEntry);
 
-				RetryStatus = FHttpRetrySystem::FRequest::EStatus::ProcessingLockout;
-
 				{
 					QUICK_SCOPE_CYCLE_COUNTER(STAT_FHttpRetrySystem_FManager_Update_OnRequestWillRetry);
 					OnRequestWillRetry().ExecuteIfBound(HttpRetryRequestEntry.Request, GetResponse(), LockoutPeriod);
 				}
 
-				RetryManagerPtr->RetryHttpRequestWithDelay(HttpRetryRequestEntry, LockoutPeriod, bSucceeded);
-				return;
+				if (RetryStatus == FHttpRetrySystem::FRequest::EStatus::Cancelled) // In case CancelRequest get called during OnRequestWillRetry callback
+				{
+					TSharedRef<FHttpRequestCommon> RequestCommon = StaticCastSharedRef<FHttpRequestCommon>(HttpRequest);
+					RequestCommon->SetFailureReason(EHttpFailureReason::Cancelled);
+					RequestCommon->SetStatus(EHttpRequestStatus::Failed);
+					HttpResponse.Reset();
+					bSucceeded = false;
+				}
+				else
+				{
+					RetryStatus = FHttpRetrySystem::FRequest::EStatus::ProcessingLockout;
+					RetryManagerPtr->RetryHttpRequestWithDelay(HttpRetryRequestEntry, LockoutPeriod, bSucceeded);
+					return;
+				}
 			}
 
 			if (HttpRetryRequestEntry.CurrentRetryCount > 0)
@@ -483,13 +493,13 @@ void FHttpRetrySystem::FManager::RetryHttpRequest(FHttpRetryRequestEntry& Reques
 		FHttpLogVerbosityTracker::Get().IncrementRetriedRequests();
 	}
 	++RequestEntry.CurrentRetryCount;
-	if (RequestEntry.Request->GetFailureReason() == EHttpFailureReason::ConnectionError)
+	if (RequestEntry.Request->LastFailureReason == EHttpFailureReason::ConnectionError)
 	{
 		++RequestEntry.CurrentRetryCountForConnectionError;
 	}
 	RequestEntry.Request->RetryStatus = FRequest::EStatus::Processing;
 
-	if (const FHttpResponsePtr Response = RequestEntry.Request->GetResponse())
+	if (const FHttpResponsePtr Response = RequestEntry.Request->LastResponse)
 	{
 		if (const int32 ResponseCode = Response->GetResponseCode(); ResponseCode < 400)
 		{
@@ -519,6 +529,10 @@ void FHttpRetrySystem::FManager::RetryHttpRequest(FHttpRetryRequestEntry& Reques
 
 void FHttpRetrySystem::FManager::RetryHttpRequestWithDelay(FManager::FHttpRetryRequestEntry& RequestEntry, float InDelay, bool bWasSucceeded)
 {
+	RequestEntry.Request->LastResponse = RequestEntry.Request->GetResponse();
+	RequestEntry.Request->LastFailureReason = RequestEntry.Request->GetFailureReason();
+	RequestEntry.Request->bLastSucceeded = bWasSucceeded;
+
 	// Timeout during lock out to keep existing behavior
 	float TimeoutOrDefault = RequestEntry.Request->GetTimeout().Get(FHttpModule::Get().GetHttpTotalTimeout());
 	if (TimeoutOrDefault != 0)
@@ -532,9 +546,13 @@ void FHttpRetrySystem::FManager::RetryHttpRequestWithDelay(FManager::FHttpRetryR
 		}
 	}
 
+	// We're sure it will be retried at this point without further access of response, clear old state, so CancelRequest can work well during this back off delay
+	TSharedRef<FHttpRequestCommon> RequestCommon = StaticCastSharedRef<FHttpRequestCommon>(RequestEntry.Request->HttpRequest);
+	RequestCommon->ClearInCaseOfRetry();
+
 	// Delay and start
 	TWeakPtr<FRequest> RequestWeakPtr(RequestEntry.Request);
-	FHttpModule::Get().GetHttpManager().AddHttpThreadTask([RequestWeakPtr, bWasSucceeded]() {
+	FHttpModule::Get().GetHttpManager().AddHttpThreadTask([RequestWeakPtr]() {
 		if (TSharedPtr<FRequest> RequestPtr = RequestWeakPtr.Pin())
 		{
 			if (TSharedPtr<FManager> RetryManagerPtr = RequestPtr->RetryManager.Pin())
@@ -545,10 +563,11 @@ void FHttpRetrySystem::FManager::RetryHttpRequestWithDelay(FManager::FHttpRetryR
 				if (EntryIndex != INDEX_NONE)
 				{
 					FManager::FHttpRetryRequestEntry* HttpRetryRequestEntry = &RetryManagerPtr->RequestList[EntryIndex];
-					// TODO: Move this into RetryHttpRequest after stabilizing the flow with CVarHttpRetrySystemNonGameThreadSupportEnabled on
-					HttpRetryRequestEntry->Request->LastResponse = HttpRetryRequestEntry->Request->GetResponse();
-					HttpRetryRequestEntry->Request->bLastSucceeded = bWasSucceeded;
-					RetryManagerPtr->RetryHttpRequest(*HttpRetryRequestEntry);
+					// Check entry state as well, in case CancelRequest->AbortRequest is still queueing in http thread tasks
+					if (HttpRetryRequestEntry->Request->RetryStatus != FHttpRetrySystem::FRequest::EStatus::Cancelled)
+					{
+						RetryManagerPtr->RetryHttpRequest(*HttpRetryRequestEntry);
+					}
 				}
 			}
 		}
@@ -619,8 +638,7 @@ float FHttpRetrySystem::FManager::GetLockoutPeriodSeconds(const FHttpRetryReques
 static FRandomStream TempRandomStream(4435261);
 
 FHttpRetrySystem::FManager::FHttpRetryRequestEntry::FHttpRetryRequestEntry(TSharedRef<FHttpRetrySystem::FRequest, ESPMode::ThreadSafe>& InRequest)
-	: bShouldCancel(false)
-	, CurrentRetryCount(0)
+	: CurrentRetryCount(0)
 	, CurrentRetryCountForConnectionError(0)
 	, RequestStartTimeAbsoluteSeconds(FPlatformTime::Seconds())
 	, Request(InRequest)
@@ -656,7 +674,7 @@ void FHttpRetrySystem::FManager::CancelRequest(TSharedRef<FHttpRetrySystem::FReq
 
 		FScopeLock ScopeLock(&RequestListLock);
 
-		// Find the existing request entry if is was previously processed.
+		// Find the existing request entry if it was previously processed.
 		bool bFound = false;
 		for (int32 i = 0; i < RequestList.Num(); ++i)
 		{
@@ -664,16 +682,14 @@ void FHttpRetrySystem::FManager::CancelRequest(TSharedRef<FHttpRetrySystem::FReq
 
 			if (EntryRef.Request == HttpRetryRequest)
 			{
-				EntryRef.bShouldCancel = true;
 				bFound = true;
 			}
 		}
-		// If we did not find the entry, likely auth failed for the request, in which case ProcessRequest does not get called.
-		// Adding it to the list and flagging as cancel will process it on next tick.
+		// If we did not find the entry, likely ProcessRequest hasn't been called. Add it to RequestList 
+		// anyway to go through the same flow of cancellation
 		if (!bFound)
 		{
 			FHttpRetryRequestEntry RetryRequestEntry(HttpRetryRequest);
-			RetryRequestEntry.bShouldCancel = true;
 			RequestList.Add(RetryRequestEntry);
 		}
 
diff --git a/Engine/Source/Runtime/Online/HTTP/Public/HttpRetrySystem.h b/Engine/Source/Runtime/Online/HTTP/Public/HttpRetrySystem.h
--- a/Engine/Source/Runtime/Online/HTTP/Public/HttpRetrySystem.h
+++ b/Engine/Source/Runtime/Online/HTTP/Public/HttpRetrySystem.h
@@ -181,6 +181,7 @@ namespace FHttpRetrySystem
 		TWeakPtr<FManager> RetryManager;
 		/** Save the last response before the retry */
 		FHttpResponsePtr LastResponse;
+		EHttpFailureReason LastFailureReason = EHttpFailureReason::None;
 		bool bLastSucceeded = false;
 
 		/** Exponential backoff curve */
@@ -231,7 +232,6 @@ protected:
 	{
 		FHttpRetryRequestEntry(TSharedRef<FRequest, ESPMode::ThreadSafe>& InRequest);
 
-		bool bShouldCancel;
 		uint32 CurrentRetryCount;
 		uint32 CurrentRetryCountForConnectionError;
 		double RequestStartTimeAbsoluteSeconds;
