Comments

#rnx
#Synced-CL 51354374

diff --git a/Engine/Source/Runtime/Online/BuildPatchServices/Private/Diffing/DiffManifests.cpp b/Engine/Source/Runtime/Online/BuildPatchServices/Private/Diffing/DiffManifests.cpp
--- a/Engine/Source/Runtime/Online/BuildPatchServices/Private/Diffing/DiffManifests.cpp
+++ b/Engine/Source/Runtime/Online/BuildPatchServices/Private/Diffing/DiffManifests.cpp
@@ -443,6 +443,9 @@ namespace BuildPatchServices
 
 		// We want a separate report for each tag, since tags can't assume files from other tags are installed and it's easy to 
 		// accumulate after if we care.
+		// We specifically do this for tags not files because the client code installs off of tags not files. If we work off of
+		// files we may end up with mulitple install tags for the file, and we can't infer which of those is installed, which
+		// means we can't get an accurate representation of what files we can use to patch off of. 
 		for (const FString& InstallTag : InstallTags)
 		{
 			UE_LOG(LogDiffManifests, Display, TEXT("%s:"), *InstallTag);
