Try to fix the concurrent issue in http retry system which could ref the destroyed retry manager instance.

#Synced-CL 51366129

diff --git a/Engine/Source/Runtime/Online/HTTP/Private/HttpRetrySystem.cpp b/Engine/Source/Runtime/Online/HTTP/Private/HttpRetrySystem.cpp
--- a/Engine/Source/Runtime/Online/HTTP/Private/HttpRetrySystem.cpp
+++ b/Engine/Source/Runtime/Online/HTTP/Private/HttpRetrySystem.cpp
@@ -80,7 +80,7 @@ float FHttpRetrySystem::FExponentialBackoffCurve::Compute(uint32 RetryNumber) co
 }
 
 FHttpRetrySystem::FRequest::FRequest(
-	TSharedRef<FManager> InManager,
+	TSharedRef<FManager, ESPMode::ThreadSafe> InManager,
 	const TSharedRef<IHttpRequest, ESPMode::ThreadSafe>& HttpRequest,
 	const FHttpRetrySystem::FRetryLimitCountSetting& InRetryLimitCountOverride,
 	const FHttpRetrySystem::FRetryTimeoutRelativeSecondsSetting& InRetryTimeoutRelativeSecondsOverride,
@@ -136,7 +136,7 @@ void FHttpRetrySystem::FRequest::BindAdaptorDelegates()
 
 bool FHttpRetrySystem::FRequest::ProcessRequest()
 {
-	TSharedRef<FRequest> RetryRequest = StaticCastSharedRef<FRequest>(AsShared());
+	TSharedRef<FRequest, ESPMode::ThreadSafe> RetryRequest = StaticCastSharedRef<FRequest>(AsShared());
 
 	OriginalUrl = HttpRequest->GetURL();
 
@@ -163,7 +163,7 @@ bool FHttpRetrySystem::FRequest::ProcessRequest()
 
 	BindAdaptorDelegates();
 
-	TSharedPtr<FManager> RetryManagerPtr = RetryManager.Pin();
+	TSharedPtr<FManager, ESPMode::ThreadSafe> RetryManagerPtr = RetryManager.Pin();
 
 	if (ensure(RetryManagerPtr))
 	{
@@ -216,7 +216,7 @@ void FHttpRetrySystem::FRequest::CancelRequest()
 
 	BindAdaptorDelegates();
 
-	if (TSharedPtr<FManager> RetryManagerPtr = RetryManager.Pin())
+	if (TSharedPtr<FManager, ESPMode::ThreadSafe> RetryManagerPtr = RetryManager.Pin())
 	{
 		RetryManagerPtr->CancelRequest(RetryRequest);
 	}
@@ -233,13 +233,13 @@ void FHttpRetrySystem::FRequest::HttpOnRequestProgress(FHttpRequestPtr InHttpReq
 
 void FHttpRetrySystem::FRequest::HttpOnProcessRequestComplete(FHttpRequestPtr InHttpRequest, FHttpResponsePtr HttpResponse, bool bSucceeded)
 {
-	TSharedPtr<FManager> RetryManagerPtr = RetryManager.Pin();
+	TSharedPtr<FManager, ESPMode::ThreadSafe> RetryManagerPtr = RetryManager.Pin();
 	if (!RetryManagerPtr)
 	{
 		return;
 	}
 
-	TSharedRef<FRequest> SelfPtr = StaticCastSharedRef<FRequest>(AsShared()); // In case no ref after removing from RetryManager
+	TSharedRef<FRequest, ESPMode::ThreadSafe> SelfPtr = StaticCastSharedRef<FRequest>(AsShared()); // In case no ref after removing from RetryManager
 
 
 	{
@@ -250,7 +250,8 @@ void FHttpRetrySystem::FRequest::HttpOnProcessRequestComplete(FHttpRequestPtr In
 		{
 			FManager::FHttpRetryRequestEntry& HttpRetryRequestEntry = RetryManagerPtr->RequestList[EntryIndex];
 
-			if (RetryStatus == FHttpRetrySystem::FRequest::EStatus::Cancelled)
+			bool bCancelled = (RetryStatus == FHttpRetrySystem::FRequest::EStatus::Cancelled) || (GetFailureReason() == EHttpFailureReason::Cancelled);
+			if (bCancelled)
 			{
 				// Do nothing here
 			}
@@ -275,7 +276,7 @@ void FHttpRetrySystem::FRequest::HttpOnProcessRequestComplete(FHttpRequestPtr In
 				RetryStatus = FHttpRetrySystem::FRequest::EStatus::Succeeded;
 			}
 
-			if (RetryStatus != FHttpRetrySystem::FRequest::EStatus::Cancelled && 
+			if (!bCancelled && 
 				RetryStatus != FHttpRetrySystem::FRequest::EStatus::FailedTimeout && 
 				RetryManagerPtr->ShouldRetry(HttpRetryRequestEntry) && 
 				RetryManagerPtr->CanRetry(HttpRetryRequestEntry))
@@ -290,7 +291,7 @@ void FHttpRetrySystem::FRequest::HttpOnProcessRequestComplete(FHttpRequestPtr In
 
 				if (RetryStatus == FHttpRetrySystem::FRequest::EStatus::Cancelled) // In case CancelRequest get called during OnRequestWillRetry callback
 				{
-					TSharedRef<FHttpRequestCommon> RequestCommon = StaticCastSharedRef<FHttpRequestCommon>(HttpRequest);
+					TSharedRef<FHttpRequestCommon, ESPMode::ThreadSafe> RequestCommon = StaticCastSharedRef<FHttpRequestCommon>(HttpRequest);
 					RequestCommon->SetFailureReason(EHttpFailureReason::Cancelled);
 					RequestCommon->SetStatus(EHttpRequestStatus::Failed);
 					HttpResponse.Reset();
@@ -332,13 +333,13 @@ void FHttpRetrySystem::FRequest::HttpOnProcessRequestComplete(FHttpRequestPtr In
 
 void FHttpRetrySystem::FRequest::HttpOnStatusCodeReceived(FHttpRequestPtr Request, int32 StatusCode)
 {
-	TSharedRef<FRequest> SelfPtr = StaticCastSharedRef<FRequest>(AsShared());
+	TSharedRef<FRequest, ESPMode::ThreadSafe> SelfPtr = StaticCastSharedRef<FRequest>(AsShared());
 	OnStatusCodeReceived().ExecuteIfBound(SelfPtr, StatusCode);
 }
 
 void FHttpRetrySystem::FRequest::HttpOnHeaderReceived(FHttpRequestPtr Request, const FString& HeaderName, const FString& NewHeaderValue)
 {
-	TSharedRef<FRequest> SelfPtr = StaticCastSharedRef<FRequest>(AsShared());
+	TSharedRef<FRequest, ESPMode::ThreadSafe> SelfPtr = StaticCastSharedRef<FRequest>(AsShared());
 	OnHeaderReceived().ExecuteIfBound(SelfPtr, HeaderName, NewHeaderValue);
 }
 
@@ -547,15 +548,15 @@ void FHttpRetrySystem::FManager::RetryHttpRequestWithDelay(FManager::FHttpRetryR
 	}
 
 	// We're sure it will be retried at this point without further access of response, clear old state, so CancelRequest can work well during this back off delay
-	TSharedRef<FHttpRequestCommon> RequestCommon = StaticCastSharedRef<FHttpRequestCommon>(RequestEntry.Request->HttpRequest);
+	TSharedRef<FHttpRequestCommon, ESPMode::ThreadSafe> RequestCommon = StaticCastSharedRef<FHttpRequestCommon>(RequestEntry.Request->HttpRequest);
 	RequestCommon->ClearInCaseOfRetry();
 
 	// Delay and start
 	TWeakPtr<FRequest> RequestWeakPtr(RequestEntry.Request);
 	FHttpModule::Get().GetHttpManager().AddHttpThreadTask([RequestWeakPtr]() {
-		if (TSharedPtr<FRequest> RequestPtr = RequestWeakPtr.Pin())
+		if (TSharedPtr<FRequest, ESPMode::ThreadSafe> RequestPtr = RequestWeakPtr.Pin())
 		{
-			if (TSharedPtr<FManager> RetryManagerPtr = RequestPtr->RetryManager.Pin())
+			if (TSharedPtr<FManager, ESPMode::ThreadSafe> RetryManagerPtr = RequestPtr->RetryManager.Pin())
 			{
 				FScopeLock ScopeLock(&RetryManagerPtr->RequestListLock);
 				// Check if it's still there in case it has been cancelled during the delay period
@@ -578,9 +579,9 @@ void FHttpRetrySystem::FManager::HttpRequestTimeoutAfterDelay(FManager::FHttpRet
 {
 	TWeakPtr<FRequest> RequestWeakPtr(RequestEntry.Request);
 	TFunction<void()> Callback([RequestWeakPtr, bWasSucceeded]() {
-		if (TSharedPtr<FRequest> RequestPtr = RequestWeakPtr.Pin())
+		if (TSharedPtr<FRequest, ESPMode::ThreadSafe> RequestPtr = RequestWeakPtr.Pin())
 		{
-			if (TSharedPtr<FManager> RetryManagerPtr = RequestPtr->RetryManager.Pin())
+			if (TSharedPtr<FManager, ESPMode::ThreadSafe> RetryManagerPtr = RequestPtr->RetryManager.Pin())
 			{
 				FScopeLock ScopeLock(&RetryManagerPtr->RequestListLock);
 				uint32 EntryIndex = RetryManagerPtr->RequestList.IndexOfByPredicate([RequestPtr](const FManager::FHttpRetryRequestEntry& Entry) { return Entry.Request == RequestPtr; });
diff --git a/Engine/Source/Runtime/Online/HTTP/Public/HttpRetrySystem.h b/Engine/Source/Runtime/Online/HTTP/Public/HttpRetrySystem.h
--- a/Engine/Source/Runtime/Online/HTTP/Public/HttpRetrySystem.h
+++ b/Engine/Source/Runtime/Online/HTTP/Public/HttpRetrySystem.h
@@ -142,8 +142,8 @@ namespace FHttpRetrySystem
 		friend class FManager;
 
 		HTTP_API FRequest(
-			TSharedRef<FManager> InManager,
-			const TSharedRef<IHttpRequest>& HttpRequest,
+			TSharedRef<FManager, ESPMode::ThreadSafe> InManager,
+			const TSharedRef<IHttpRequest, ESPMode::ThreadSafe>& HttpRequest,
 			const FRetryLimitCountSetting& InRetryLimitCountOverride = FRetryLimitCountSetting(),
 			const FRetryTimeoutRelativeSecondsSetting& InRetryTimeoutRelativeSecondsOverride = FRetryTimeoutRelativeSecondsSetting(),
 			const FRetryResponseCodes& InRetryResponseCodes = FRetryResponseCodes(),
