Add CVars for behavior changes(cancel during will retry delegate or back off) of submitted changelist 51424098.

#Synced-CL 51366129

diff --git a/Engine/Source/Programs/WebTests/Private/TestHttp.cpp b/Engine/Source/Programs/WebTests/Private/TestHttp.cpp
--- a/Engine/Source/Programs/WebTests/Private/TestHttp.cpp
+++ b/Engine/Source/Programs/WebTests/Private/TestHttp.cpp
@@ -52,6 +52,7 @@ extern TAutoConsoleVariable<FString> CVarHttpUrlPatternsToMockResponse;
 extern TAutoConsoleVariable<FString> CVarHttpUrlPatternsToDisableFailedLog;
 extern TAutoConsoleVariable<FString> CVarHttpUrlPatternsToInspect;
 extern TAutoConsoleVariable<bool> CVarHttpLogJsonContentOnly;
+extern TAutoConsoleVariable<bool> CVarHttpCancelSupportDuringBackOffEnabled;
 
 class FMockHttpModule : public FHttpModule
 {
@@ -2684,6 +2685,11 @@ TEST_CASE_METHOD(FWaitUntilCompleteHttpFixture, "Request can be canceled during
 		return;
 	}
 
+	if (!CVarHttpCancelSupportDuringBackOffEnabled.GetValueOnAnyThread())
+	{
+		return;
+	}
+
 	DisableWarningsInThisTest();
 
 	TSharedRef<IHttpRequest> HttpRequest = HttpRetryManager->CreateRequest(
diff --git a/Engine/Source/Runtime/Online/HTTP/Private/Curl/CurlHttp.cpp b/Engine/Source/Runtime/Online/HTTP/Private/Curl/CurlHttp.cpp
--- a/Engine/Source/Runtime/Online/HTTP/Private/Curl/CurlHttp.cpp
+++ b/Engine/Source/Runtime/Online/HTTP/Private/Curl/CurlHttp.cpp
@@ -36,6 +36,13 @@ TAutoConsoleVariable<bool> CVarHttpAlwaysSetContentLengthEnabled(
 
 extern TAutoConsoleVariable<bool> CVarHttpCurlReuseConnectionEnabled;
 
+TAutoConsoleVariable<bool> CVarHttpCurlResetRequestCompletedOnRetryEnabled(
+	TEXT("http.CurlResetRequestCompletedOnRetryEnabled"),
+	true,
+	TEXT("When enabled, bCurlRequestCompleted is reset to false in ClearInCaseOfRetry so the request can be reused for retrying."),
+	ECVF_SaveForNextBoot
+);
+
 #if WITH_SSL
 static int SslCertVerify(int PreverifyOk, X509_STORE_CTX* Context)
 {
@@ -1101,7 +1108,10 @@ void FCurlHttpRequest::ClearInCaseOfRetry()
 	LastReportedBytesRead = 0;
 	TotalBytesRead = 0;
 	bAnyHttpActivity = false;
-	bCurlRequestCompleted = false;
+	if (CVarHttpCurlResetRequestCompletedOnRetryEnabled.GetValueOnAnyThread())
+	{
+		bCurlRequestCompleted = false;
+	}
 
 	// Clear the info cache log so we don't output messages from previous requests when reusing/retrying a request
 	{
diff --git a/Engine/Source/Runtime/Online/HTTP/Private/HttpRetrySystem.cpp b/Engine/Source/Runtime/Online/HTTP/Private/HttpRetrySystem.cpp
--- a/Engine/Source/Runtime/Online/HTTP/Private/HttpRetrySystem.cpp
+++ b/Engine/Source/Runtime/Online/HTTP/Private/HttpRetrySystem.cpp
@@ -18,6 +18,20 @@
 
 LLM_DEFINE_TAG(HTTP);
 
+TAutoConsoleVariable<bool> CVarHttpCancelSupportDuringWillRetryDelegate(
+	TEXT("http.CancelSupportDuringWillRetryDelegate"),
+	true,
+	TEXT("When enabled, http request can be cancelled during WillRetry delegate."),
+	ECVF_SaveForNextBoot
+);
+
+TAutoConsoleVariable<bool> CVarHttpCancelSupportDuringBackOffEnabled(
+	TEXT("http.CancelSupportDuringBackOffEnabled"),
+	false,
+	TEXT("When enabled, http request can be cancelled during back off."),
+	ECVF_SaveForNextBoot
+);
+
 namespace FHttpRetrySystem
 {
 
@@ -289,7 +303,8 @@ void FHttpRetrySystem::FRequest::HttpOnProcessRequestComplete(FHttpRequestPtr In
 					OnRequestWillRetry().ExecuteIfBound(HttpRetryRequestEntry.Request, GetResponse(), LockoutPeriod);
 				}
 
-				if (RetryStatus == FHttpRetrySystem::FRequest::EStatus::Cancelled) // In case CancelRequest get called during OnRequestWillRetry callback
+				if (CVarHttpCancelSupportDuringWillRetryDelegate.GetValueOnAnyThread() &&
+					RetryStatus == FHttpRetrySystem::FRequest::EStatus::Cancelled) // In case CancelRequest get called during OnRequestWillRetry callback
 				{
 					TSharedRef<FHttpRequestCommon, ESPMode::ThreadSafe> RequestCommon = StaticCastSharedRef<FHttpRequestCommon>(HttpRequest);
 					RequestCommon->SetFailureReason(EHttpFailureReason::Cancelled);
@@ -547,9 +562,12 @@ void FHttpRetrySystem::FManager::RetryHttpRequestWithDelay(FManager::FHttpRetryR
 		}
 	}
 
-	// We're sure it will be retried at this point without further access of response, clear old state, so CancelRequest can work well during this back off delay
-	TSharedRef<FHttpRequestCommon, ESPMode::ThreadSafe> RequestCommon = StaticCastSharedRef<FHttpRequestCommon>(RequestEntry.Request->HttpRequest);
-	RequestCommon->ClearInCaseOfRetry();
+	if (CVarHttpCancelSupportDuringBackOffEnabled.GetValueOnAnyThread())
+	{
+		// We're sure it will be retried at this point without further access of response, clear old state, so CancelRequest can work well during this back off delay
+		TSharedRef<FHttpRequestCommon, ESPMode::ThreadSafe> RequestCommon = StaticCastSharedRef<FHttpRequestCommon>(RequestEntry.Request->HttpRequest);
+		RequestCommon->ClearInCaseOfRetry();
+	}
 
 	// Delay and start
 	TWeakPtr<FRequest> RequestWeakPtr(RequestEntry.Request);
@@ -565,7 +583,8 @@ void FHttpRetrySystem::FManager::RetryHttpRequestWithDelay(FManager::FHttpRetryR
 				{
 					FManager::FHttpRetryRequestEntry* HttpRetryRequestEntry = &RetryManagerPtr->RequestList[EntryIndex];
 					// Check entry state as well, in case CancelRequest->AbortRequest is still queueing in http thread tasks
-					if (HttpRetryRequestEntry->Request->RetryStatus != FHttpRetrySystem::FRequest::EStatus::Cancelled)
+					if (!CVarHttpCancelSupportDuringBackOffEnabled.GetValueOnAnyThread() ||
+						HttpRetryRequestEntry->Request->RetryStatus != FHttpRetrySystem::FRequest::EStatus::Cancelled)
 					{
 						RetryManagerPtr->RetryHttpRequest(*HttpRetryRequestEntry);
 					}
