Gameplay integration for Downlink Manager

#redirectsaresafe
#rnx

#Synced-CL 51635240

diff --git a/Engine/Source/Runtime/Engine/Engine.Build.cs b/Engine/Source/Runtime/Engine/Engine.Build.cs
--- a/Engine/Source/Runtime/Engine/Engine.Build.cs
+++ b/Engine/Source/Runtime/Engine/Engine.Build.cs
@@ -174,6 +174,7 @@ public class Engine : ModuleRules
 				"IoStoreOnDemand",
 				"HTTP",
 				"GeometryCore",
+				"DownlinkBandwidthManager",
 			}
 		);
 
diff --git a/Engine/Source/Runtime/Engine/Private/PlayerController.cpp b/Engine/Source/Runtime/Engine/Private/PlayerController.cpp
--- a/Engine/Source/Runtime/Engine/Private/PlayerController.cpp
+++ b/Engine/Source/Runtime/Engine/Private/PlayerController.cpp
@@ -67,6 +67,7 @@
 #include "WorldPartition/WorldPartitionSubsystem.h"
 #include "PBDRigidsSolver.h"
 #include "PhysicsEngine/PhysicsSettings.h"
+#include "DownlinkBandwidthManager.h"
 
 #include "Iris/ReplicationSystem/ReplicationSystem.h"
 #include "Net/Iris/ReplicationSystem/EngineReplicationBridge.h"
@@ -1436,6 +1437,10 @@ void APlayerController::OnActorChannelOpen(FInBunch& InBunch, UNetConnection* Co
 		{
 			// main connection PlayerController
 			Connection->HandleClientPlayer(this, Connection); 
+
+			// sets the reserved bandwidth on the bandwidth manager for gameplay in Kbps
+			FDownlinkBandwidthManager& DBManager = FDownlinkBandwidthManager::Get();
+			DBManager.SetGameplayBandwidthCap(Connection->Driver->MaxInternetClientRate / UE::BandwidthManager::KilobitsToBytesConversionValue);
 		}
 		else
 		{
diff --git a/Engine/Source/Runtime/Experimental/IoStore/OnDemand/Private/OnDemandHttpIoDispacher.cpp b/Engine/Source/Runtime/Experimental/IoStore/OnDemand/Private/OnDemandHttpIoDispacher.cpp
--- a/Engine/Source/Runtime/Experimental/IoStore/OnDemand/Private/OnDemandHttpIoDispacher.cpp
+++ b/Engine/Source/Runtime/Experimental/IoStore/OnDemand/Private/OnDemandHttpIoDispacher.cpp
@@ -350,6 +350,13 @@ static FAutoConsoleVariableRef CVar_IaxHttpForceInitialConnectionCheck(
 	TEXT("When true the IAX dispatcher will be forced to attempt to connect to at least one hostgroup before servicing requests")
 );
 
+bool GUseDownlinkBandwidthManager = false;
+static FAutoConsoleVariableRef CVar_UseDLBWManager(
+	TEXT("iax.UseDLBWManager"),
+	GUseDownlinkBandwidthManager,
+	TEXT("Toggle state for whether or not the HttpDispatcher will use the Downlink Bandwidth Manager.")
+);
+
 ////////////////////////////////////////////////////////////////////////////////
 int32 CalculateRequestQueueLengthMax()
 {
@@ -377,7 +384,7 @@ int32 GetHttpRateLimitKiBPerSecond()
 	}
 #endif // UE_ALLOW_DISABLE_HTTP_LIMITS
 
-	return (FDownlinkBandwidthManager::IsEnabled()) ? DLBWHttpRateLimitKiBPerSecond.load() : GIasHttpRateLimitKiBPerSecond;
+	return (GUseDownlinkBandwidthManager) ? DLBWHttpRateLimitKiBPerSecond.load() : GIasHttpRateLimitKiBPerSecond;
 }
 
 ////////////////////////////////////////////////////////////////////////////////
@@ -1436,6 +1443,9 @@ FHttpDispatcher::~FHttpDispatcher()
 		OnRecreateHttpClientHandle.Reset();
 	}
 
+	FDownlinkBandwidthManager& DBManager = FDownlinkBandwidthManager::Get();
+	DBManager.UnregisterMonitoredService(BandwidthManagerServiceHandle);
+
 	Thread.Reset();
 }
 
@@ -1447,22 +1457,8 @@ bool FHttpDispatcher::Init()
 			WakeUp->Trigger();
 		});
 
-	if (FDownlinkBandwidthManager::IsEnabled())
-	{
-		FDownlinkBandwidthManager& DBManager = FDownlinkBandwidthManager::Get();
-		UE::BandwidthManager::FManagedServiceConfigurableData BandwidthConfig;
-		BandwidthConfig.ComparisonValue = MakeShared<int32>(GIasHttpRateLimitKiBPerSecond);
-		BandwidthConfig.IdleGameplayPriority = UE::BandwidthManager::EBandwidthPriority::HighPriority;
-		BandwidthConfig.InGameplayPriority = UE::BandwidthManager::EBandwidthPriority::LowPriority;
-
-		BandwidthManagerServiceHandle = DBManager.RegisterMonitoredService(TEXT("HttpDispatcher"), BandwidthConfig, true,
-			[](int32 AllottedBandwidthKbps) 
-			{
-				// 1 Kilobyte is equal to (1000 / 1024) kibibytes
-				DLBWHttpRateLimitKiBPerSecond = static_cast<int32>((float)AllottedBandwidthKbps * (1000.0f / 1024.0f));
-			});
-	}
-
+	// Sets an initial value in case of delay for bandwidth manager registration
+	DLBWHttpRateLimitKiBPerSecond = static_cast<int32>(GIasHttpRateLimitKiBPerSecond);
 	return TryCreateHttpClient();
 }
 
@@ -1470,12 +1466,6 @@ uint32 FHttpDispatcher::Run()
 {
 	check(HttpClient.IsValid());
 
-	if (FDownlinkBandwidthManager::IsEnabled())
-	{
-		FDownlinkBandwidthManager& DBManager = FDownlinkBandwidthManager::Get();
-		DBManager.SetActivationForService(BandwidthManagerServiceHandle, true);
-	}
-
 	if (GIaxHttpForceInitialConnectionCheck)
 	{
 		bool bFirstAttempt = false;
@@ -1502,6 +1492,28 @@ uint32 FHttpDispatcher::Run()
 
 	for (;;)
 	{
+		if (GUseDownlinkBandwidthManager && BandwidthManagerServiceHandle == UE_INVALID_BANDWIDTH_SERVICE_HANDLE)
+		{
+			FDownlinkBandwidthManager& DBManager = FDownlinkBandwidthManager::Get();
+			UE::BandwidthManager::FManagedServiceConfigurableData BandwidthConfig;
+			BandwidthConfig.ComparisonValue = MakeShared<int32>(GIasHttpRateLimitKiBPerSecond);
+			BandwidthConfig.IdleGameplayPriority = UE::BandwidthManager::EBandwidthPriority::HighPriority;
+			BandwidthConfig.InGameplayPriority = UE::BandwidthManager::EBandwidthPriority::LowPriority;
+
+			BandwidthManagerServiceHandle = DBManager.RegisterMonitoredService(TEXT("HttpDispatcher"), BandwidthConfig, true,
+				[this](int32 AllottedBandwidthKbps)
+				{
+					// 1 Kilobyte is equal to (1000 / 1024) kibibytes
+					constexpr float KilobyteToKibibyte = (1000.0f / 1024.0f);
+					DLBWHttpRateLimitKiBPerSecond = static_cast<int32>((float)AllottedBandwidthKbps * KilobyteToKibibyte);
+				});
+		}
+		else if (GUseDownlinkBandwidthManager)
+		{
+			FDownlinkBandwidthManager& DBManager = FDownlinkBandwidthManager::Get();
+			DBManager.SetActivationForService(BandwidthManagerServiceHandle, true);
+		}
+
 		UpdateThreadPriorityIfNeeded();
 
 		Tick();
@@ -1548,7 +1560,7 @@ void FHttpDispatcher::Stop()
 		WakeUp->Trigger();
 	}
 
-	if (FDownlinkBandwidthManager::IsEnabled())
+	if (GUseDownlinkBandwidthManager)
 	{
 		FDownlinkBandwidthManager& DBManager = FDownlinkBandwidthManager::Get();
 		DBManager.SetActivationForService(BandwidthManagerServiceHandle, false);
@@ -1557,11 +1569,8 @@ void FHttpDispatcher::Stop()
 
 void FHttpDispatcher::Exit()
 {
-	if (FDownlinkBandwidthManager::IsEnabled())
-	{
-		FDownlinkBandwidthManager& DBManager = FDownlinkBandwidthManager::Get();
-		DBManager.UnregisterMonitoredService(BandwidthManagerServiceHandle);
-	}
+	FDownlinkBandwidthManager& DBManager = FDownlinkBandwidthManager::Get();
+	DBManager.UnregisterMonitoredService(BandwidthManagerServiceHandle);
 }
 
 void FHttpDispatcher::Tick()
diff --git a/Engine/Source/Runtime/Online/Experimental/DownlinkBandwidthManager/Private/DownlinkBandwidthManager.cpp b/Engine/Source/Runtime/Online/Experimental/DownlinkBandwidthManager/Private/DownlinkBandwidthManager.cpp
--- a/Engine/Source/Runtime/Online/Experimental/DownlinkBandwidthManager/Private/DownlinkBandwidthManager.cpp
+++ b/Engine/Source/Runtime/Online/Experimental/DownlinkBandwidthManager/Private/DownlinkBandwidthManager.cpp
@@ -2,20 +2,28 @@
 
 #include "DownlinkBandwidthManager.h"
 #include "HAL/IConsoleManager.h"
-#include "Misc/ConfigCacheIni.h"
+#include "Math/NumericLimits.h"
+//#include "Misc/ConfigCacheIni.h"
 
 DEFINE_LOG_CATEGORY(LogDownlinkBandwidthManager);
 
 namespace DLBW
 {
 	/// Console Variables to determine Manager behavior
-	int32 GDefaultBandwidthToAllocate = 300000;
+	int32 GDefaultBandwidthToAllocate = 300000; // 300 Mbps
 	static FAutoConsoleVariableRef CVarDLBWDefaultBandwidthToAllocate(
 		TEXT("net.DLBW.DefaultBandwidth"),
 		GDefaultBandwidthToAllocate,
 		TEXT("Default value to use for bandwidth cap distribution in Kbps. Example: Average american home internet is 200-300Mbps, so we would use 200,000 - 300,000 for this value."),
 		ECVF_Default);
 
+	int32 GMinRequiredBandwidthDuringGameplay = 50000; // 50 Mbps
+	static FAutoConsoleVariableRef CVarDLBWMinRequiredBandwidthDuringGameplay(
+		TEXT("net.DLBW.MinGameplayBandwidth"),
+		GMinRequiredBandwidthDuringGameplay,
+		TEXT("Default value to use for bandwidth cap distribution in Kbps. Example: Average american home internet is 200-300Mbps, so we would use 200,000 - 300,000 for this value."),
+		ECVF_Default);
+
 	static TAutoConsoleVariable<int32> CVarDLBWPriorityWeightExponent(
 		TEXT("net.DLBW.PriorityWeightExponent"),
 		3,
@@ -28,12 +36,7 @@ namespace DLBW
 		TEXT("Dealy Time, per frame cycle to poll for changes needed in distribution."),
 		ECVF_Default);
 
-	static TAutoConsoleVariable<bool> CVarDLBWEnabled(
-		TEXT("net.DLBW.Enable"),
-		false,
-		TEXT("Toggles the activation state for the Downlink Bandwidth Manager."));
-
-	bool GEnableBandwidthLogging = false;
+	static bool GEnableBandwidthLogging = false;
 	static FAutoConsoleVariableRef CVarDLBWEnableBandwidthLogging(
 		TEXT("net.DLBW.EnableLogging"),
 		GEnableBandwidthLogging,
@@ -73,11 +76,6 @@ FDownlinkBandwidthManager& FDownlinkBandwidthManager::Get()
 	return *Instance.Get();
 }
 
-bool FDownlinkBandwidthManager::IsEnabled()
-{
-	return DLBW::CVarDLBWEnabled.GetValueOnAnyThread();
-}
-
 /// Member Functions
 FDownlinkBandwidthManager::FDownlinkBandwidthManager()
 {
@@ -115,6 +113,7 @@ void FDownlinkBandwidthManager::SetIncomingBandwidthCap(int32 InKbps)
 	if (CapKbps != InKbps)
 	{
 		CapKbps = InKbps;
+
 		// a change to the cap value means we need to perform a full re-distribution
 		GeneratePriorityBucketMap();
 	}
@@ -153,22 +152,39 @@ UE::BandwidthManager::FServiceHandleID FDownlinkBandwidthManager::RegisterMonito
 }
 
 
-void FDownlinkBandwidthManager::UnregisterMonitoredService(const UE::BandwidthManager::FServiceHandleID& ServiceID)
+void FDownlinkBandwidthManager::UnregisterMonitoredService(UE::BandwidthManager::FServiceHandleID& ServiceID)
 {
 	UE::TScopeLock Lock(Mutex);
 
-	TrackedServices.Remove(ServiceID);
-	GeneratePriorityBucketMap();
+	if (ServiceID == UE_INVALID_BANDWIDTH_SERVICE_HANDLE)
+	{
+		return;
+	}
+
+	// Only regenerate if we actually remove a service
+	if (TrackedServices.Remove(ServiceID) != 0)
+	{
+		ServiceID = UE_INVALID_BANDWIDTH_SERVICE_HANDLE;
+		GeneratePriorityBucketMap();
+	}
 }
 
 void FDownlinkBandwidthManager::SetActivationForService(const UE::BandwidthManager::FServiceHandleID& ServiceID, bool ActiveState)
 {
 	UE::TScopeLock Lock(Mutex);
 
+	if (ServiceID == UE_INVALID_BANDWIDTH_SERVICE_HANDLE)
+	{
+		return;
+	}
+
 	if (FManagedService* Service = TrackedServices.Find(ServiceID))
 	{
-		Service->bActiveConnection = ActiveState;
-		GeneratePriorityBucketMap();
+		if (Service->bActiveConnection != ActiveState)
+		{
+			Service->bActiveConnection = ActiveState;
+			GeneratePriorityBucketMap();
+		}
 	}
 }
 
@@ -176,26 +192,54 @@ void FDownlinkBandwidthManager::ChangeServicePriority(const UE::BandwidthManager
 {
 	UE::TScopeLock Lock(Mutex);
 
+	if (ServiceID == UE_INVALID_BANDWIDTH_SERVICE_HANDLE)
+	{
+		return;
+	}
+
 	if (FManagedService* Service = TrackedServices.Find(ServiceID))
 	{
+		bool DidPriorityChange = false;
 		switch (PriorityState) 
 		{
 		case UE::BandwidthManager::EBandwidthDistributionPriorityState::IdleGameplay:
+			DidPriorityChange = Service->ConfigData.IdleGameplayPriority != NewPriority;
 			Service->ConfigData.IdleGameplayPriority = NewPriority;
 			break;
 		case UE::BandwidthManager::EBandwidthDistributionPriorityState::InGameplay:
+			DidPriorityChange = Service->ConfigData.InGameplayPriority != NewPriority;
 			Service->ConfigData.InGameplayPriority = NewPriority;
 			break;
 		default:
+			DidPriorityChange = Service->ConfigData.InGameplayPriority != NewPriority
+				|| Service->ConfigData.IdleGameplayPriority != NewPriority;
 			Service->ConfigData.IdleGameplayPriority = NewPriority;
 			Service->ConfigData.InGameplayPriority = NewPriority;
 			break;
 		}
 
-		GeneratePriorityBucketMap();
+		if (DidPriorityChange)
+		{
+			GeneratePriorityBucketMap();
+		}
 	}
 }
 
+int32 FDownlinkBandwidthManager::GetCurrentMaxAvailableBandwidthInBytes()
+{
+	int32 AvailableBandwithKbps = FMath::Abs(CapKbps);
+
+	if (CurrentGameplayPriorityState == UE::BandwidthManager::EBandwidthDistributionPriorityState::InGameplay)
+	{
+		int32 ReducedBandwidth = CapKbps - ReservedGameplayBandwidthKbps;
+		AvailableBandwithKbps = (ReducedBandwidth <= 0) ? DLBW::GMinRequiredBandwidthDuringGameplay : ReducedBandwidth;
+	}
+
+	//overflow protection for byte conversion
+	constexpr int32 IntMax = TNumericLimits<int32>::Max(); 
+	return (AvailableBandwithKbps > IntMax / UE::BandwidthManager::KilobitsToBytesConversionValue) ? IntMax : AvailableBandwithKbps * UE::BandwidthManager::KilobitsToBytesConversionValue;
+}
+
 void FDownlinkBandwidthManager::SetDistributionGameplayState(UE::BandwidthManager::EBandwidthDistributionPriorityState State)
 {
 	UE::TScopeLock Lock(Mutex);
@@ -230,7 +274,7 @@ void FDownlinkBandwidthManager::GeneratePriorityBucketMap()
 {
 	UE::TScopeLock Lock(Mutex);
 
-	int32 BandwidthBudgetBytes = CapKbps;
+	int32 BandwidthBudgetBytes = GetCurrentMaxAvailableBandwidthInBytes();
 	if (BandwidthBudgetBytes <= 0)
 	{
 		UE_LOG(LogDownlinkBandwidthManager, Warning, TEXT("%s - BandwidthBudgetPerFrame was set to less than 0. Value: %i, CapKbps: %i"),
@@ -259,7 +303,6 @@ void FDownlinkBandwidthManager::GeneratePriorityBucketMap()
 	}
 	
 	// we do not need an exponent larger than 255, so this lets us make sure that value is not too large and within the limits of the formula
-	int32 SanitizedExponent = FMath::Min<int32>(DLBW::CVarDLBWPriorityWeightExponent.GetValueOnAnyThread(), UINT8_MAX);
 	uint8 WeightBiasExponent = FMath::Max<uint8>(1, static_cast<uint8>(DLBW::CVarDLBWPriorityWeightExponent.GetValueOnAnyThread()));
 	const int32 BucketCount = PriorityBucketMapping.Num();
 	int32 BucketNumber = 0;
@@ -312,7 +355,7 @@ void FDownlinkBandwidthManager::DistributeBandwidth()
 {
 	UE::TScopeLock Lock(Mutex);
 
-	int32 BandwidthBudgetBytes = CapKbps;
+	int32 BandwidthBudgetBytes = GetCurrentMaxAvailableBandwidthInBytes();
 	if (BandwidthBudgetBytes <= 0)
 	{
 		UE_LOG(LogDownlinkBandwidthManager, Warning, TEXT("%s - BandwidthBudgetPerFrame was set to less than 0. Value: %i, CapKbps: %i"), 
@@ -347,7 +390,7 @@ void FDownlinkBandwidthManager::DistributeBandwidth()
 			int32 CurrentServiceAbsoluteBandwidth = PriorityBucket->BandwidthAllocationForPriority / PriorityBucket->NumberOfServicesAtPriority;
 
 			// If reducing BandwidthBudgetBytes by CurrentServiceAbsoluteBandwidth would yield less than 0 (rollover) then use the smallest
-			int32 BandwidthToUse = FMath::Min(BandwidthBudgetBytes,CurrentServiceAbsoluteBandwidth);
+			int32 BandwidthToUse = FMath::Min(BandwidthBudgetBytes,CurrentServiceAbsoluteBandwidth) / UE::BandwidthManager::KilobitsToBytesConversionValue; //bytes to Kb
 			if (ServiceRef.ConfigData.ComparisonValue.IsValid())
 			{
 				// Then take the minimum bandwidth between the "need" and available
@@ -360,7 +403,7 @@ void FDownlinkBandwidthManager::DistributeBandwidth()
 			}
 
 			// reduce total amount by amount distributed
-			BandwidthBudgetBytes -= ServiceRef.AllotedBandwidthKbps;
+			BandwidthBudgetBytes -= ServiceRef.AllotedBandwidthKbps * UE::BandwidthManager::KilobitsToBytesConversionValue;
 
 			if (ServiceRef.OnBandwidthUpdate.IsSet() && ServiceRef.AllotedBandwidthKbps > 0)
 			{
@@ -380,27 +423,27 @@ void FDownlinkBandwidthManager::Log()
 	UE_LOG(LogDownlinkBandwidthManager, Log, TEXT("FDownlinkBandwidthManager Logging Breakdown:\n"));
 
 	UE_LOG(LogDownlinkBandwidthManager, Log, TEXT("FDownlinkBandwidthManager Internal Data:"));
-	UE_LOG(LogDownlinkBandwidthManager, Log, TEXT("Has the manager been enabled: %s"), (DLBW::CVarDLBWEnabled.GetValueOnAnyThread()) ? TEXT("True") : TEXT("False"));
-	UE_LOG(LogDownlinkBandwidthManager, Log, TEXT("Default Bandwidth Cap in Kbps: %i"), DLBW::GDefaultBandwidthToAllocate);
-	UE_LOG(LogDownlinkBandwidthManager, Log, TEXT("Current Bandwidth Cap in Kbps: %i"), CapKbps);
+	UE_LOG(LogDownlinkBandwidthManager, Log, TEXT("Default Bandwidth Cap in Kbps: %i"), static_cast<int>(DLBW::GDefaultBandwidthToAllocate));
+	UE_LOG(LogDownlinkBandwidthManager, Log, TEXT("Minimum Bandwidth Cap during gameplay in Kbps: %i"), static_cast<int>(DLBW::GMinRequiredBandwidthDuringGameplay));
+	UE_LOG(LogDownlinkBandwidthManager, Log, TEXT("Current Bandwidth Cap in Kbps: %i"), static_cast<int>(CapKbps));
 	UE_LOG(LogDownlinkBandwidthManager, Log, TEXT("Current Gameplay Priority State: %s"), 
 		(CurrentGameplayPriorityState == UE::BandwidthManager::EBandwidthDistributionPriorityState::IdleGameplay) ? TEXT("Outside Normal Gameplay") : TEXT("In Active Gameplay"));
 	UE_LOG(LogDownlinkBandwidthManager, Log, TEXT("Tick Polling in seconds: %0.3f"), DLBW::CVarDLBWDistributionPolling.GetValueOnAnyThread());
 
 	UE_LOG(LogDownlinkBandwidthManager, Log, TEXT("FDownlinkBandwidthManager Prioritization:"));
-	UE_LOG(LogDownlinkBandwidthManager, Log, TEXT("Priority Weight Exponent: %i"), DLBW::CVarDLBWPriorityWeightExponent.GetValueOnAnyThread());
+	UE_LOG(LogDownlinkBandwidthManager, Log, TEXT("Priority Weight Exponent: %i"), static_cast<int>(DLBW::CVarDLBWPriorityWeightExponent.GetValueOnAnyThread()));
 
 	for (TPair<UE::BandwidthManager::EBandwidthPriority, UE::BandwidthManager::FPriorityBucket>& BucketData : PriorityBucketMapping)
 	{
 		UE_LOG(LogDownlinkBandwidthManager, Log, TEXT("Bucket Priority Value: %s"), *LexToString(BucketData.Key));
-		UE_LOG(LogDownlinkBandwidthManager, Log, TEXT("Bucket Service Count: %i"), BucketData.Value.NumberOfServicesAtPriority);
-		UE_LOG(LogDownlinkBandwidthManager, Log, TEXT("Bandwidth for this Bucket: %i"), BucketData.Value.BandwidthAllocationForPriority);
+		UE_LOG(LogDownlinkBandwidthManager, Log, TEXT("Bucket Service Count: %u"), static_cast<unsigned int>(BucketData.Value.NumberOfServicesAtPriority));
+		UE_LOG(LogDownlinkBandwidthManager, Log, TEXT("Bandwidth for this Bucket: %i"), static_cast<int>(BucketData.Value.BandwidthAllocationForPriority));
 	}
 
 	UE_LOG(LogDownlinkBandwidthManager, Log, TEXT("FDownlinkBandwidthManager Tracked Services:"));
 	for (TPair<UE::BandwidthManager::FServiceHandleID, FManagedService>& ServiceData : TrackedServices)
 	{
-		UE_LOG(LogDownlinkBandwidthManager, Log, TEXT("Tracked Service ID: %i, Name: %s"), ServiceData.Key, *ServiceData.Value.DebugName);
+		UE_LOG(LogDownlinkBandwidthManager, Log, TEXT("Tracked Service ID: %i, Name: %s"), static_cast<int>(ServiceData.Key), *ServiceData.Value.DebugName);
 		UE_LOG(LogDownlinkBandwidthManager, Log, TEXT("Service Gameplay Priority Value: %s"), *LexToString(ServiceData.Value.ConfigData.InGameplayPriority));
 		UE_LOG(LogDownlinkBandwidthManager, Log, TEXT("Service Gameplay Priority Value: %s"), *LexToString(ServiceData.Value.ConfigData.IdleGameplayPriority));
 		UE_LOG(LogDownlinkBandwidthManager, Log, TEXT("Is Service Currently Active: %s"), (ServiceData.Value.bActiveConnection) ? TEXT("True") : TEXT("False"));
@@ -410,10 +453,10 @@ void FDownlinkBandwidthManager::Log()
 		}
 		else
 		{
-			UE_LOG(LogDownlinkBandwidthManager, Log, TEXT("Service Comparision Value: %i"), ServiceData.Value.ConfigData.ComparisonValue.Get());
+			UE_LOG(LogDownlinkBandwidthManager, Log, TEXT("Service Comparision Value: %i"), static_cast<int>(*ServiceData.Value.ConfigData.ComparisonValue.Get()));
 		}
 
-		UE_LOG(LogDownlinkBandwidthManager, Log, TEXT("Service Alloted Bandwidth in Kbps: %i"), ServiceData.Value.AllotedBandwidthKbps);
+		UE_LOG(LogDownlinkBandwidthManager, Log, TEXT("Service Alloted Bandwidth in Kbps: %i"), static_cast<int>(ServiceData.Value.AllotedBandwidthKbps));
 		UE_LOG(LogDownlinkBandwidthManager, Log, TEXT("Service Bound Callback: %s\n"), (ServiceData.Value.OnBandwidthUpdate.IsSet()) ? TEXT("Set") : TEXT("Unset"));
 	}
 
diff --git a/Engine/Source/Runtime/Online/Experimental/DownlinkBandwidthManager/Public/DownlinkBandwidthManager.h b/Engine/Source/Runtime/Online/Experimental/DownlinkBandwidthManager/Public/DownlinkBandwidthManager.h
--- a/Engine/Source/Runtime/Online/Experimental/DownlinkBandwidthManager/Public/DownlinkBandwidthManager.h
+++ b/Engine/Source/Runtime/Online/Experimental/DownlinkBandwidthManager/Public/DownlinkBandwidthManager.h
@@ -57,6 +57,9 @@ namespace UE::BandwidthManager
 		int32 BandwidthAllocationForPriority = 0;
 		uint8 NumberOfServicesAtPriority = 0;
 	};
+
+	// Conversion value for Kb to bytes
+	constexpr int32 KilobitsToBytesConversionValue = 1000 / 8;
 }
 
 struct FManagedService
@@ -94,9 +97,6 @@ public:
 	// --- Singleton access ---
 	UE_API static FDownlinkBandwidthManager& Get();
 
-	// --- CVar Enabled Check ---
-	UE_API static bool IsEnabled();
-
 	// --- Lifecycle (called by module) ---
 	void Initialize();
 	void Shutdown();
@@ -108,8 +108,8 @@ public:
 		bool StartActive = false,
 		FOnBandwidthUpdate&& Callback = nullptr);
 
-	// Removes service from manager considerations
-	UE_API void UnregisterMonitoredService(const UE::BandwidthManager::FServiceHandleID& ServiceID);
+	// Removes service from manager considerations and invalidates ServiceID after
+	UE_API void UnregisterMonitoredService(UE::BandwidthManager::FServiceHandleID& ServiceID);
 
 	// Determines if service should be considered for distribution
 	UE_API void SetActivationForService(const UE::BandwidthManager::FServiceHandleID& ServiceID, bool ActiveState);
@@ -119,11 +119,13 @@ public:
 
 	// Ways for the game to communicate back or check to see what priority distribution pattern the manager is in
 	UE_API void SetDistributionGameplayState(UE::BandwidthManager::EBandwidthDistributionPriorityState State);
-	UE::BandwidthManager::EBandwidthDistributionPriorityState GetDistributionGameplayState() const { return CurrentGameplayPriorityState; };
-
+	
 	// Sets the bandwidth Cap for the Manager to distribute across all services
 	UE_API void SetIncomingBandwidthCap(int32 InKbps);
 
+	// Sets the Gameplay bandwidth for the Manager to reserve during the gameplay state
+	void SetGameplayBandwidthCap(int32 InKbps) { ReservedGameplayBandwidthKbps = InKbps; };
+
 private:
 	// Called once per frame/tick to apply an updated allotment based on active services
 	void Tick();
@@ -137,6 +139,9 @@ private:
 	// Helper function to get priority based on CurrentGameplayPriority
 	UE::BandwidthManager::EBandwidthPriority GetPriorityValueForService(const FManagedService& MonitoredService) const;
 
+	// Helper function to get the current available bandwidth to distribute based on gameplay state
+	int32 GetCurrentMaxAvailableBandwidthInBytes();
+
 	// Allows for logging output of FDownlinkBandwidthManager, enabled via CVarDLBWEnableBandwidthLogging
 	void Log();
 
@@ -151,6 +156,9 @@ private:
 	// The hard cap used for determining bandwidth distribution
 	int32 CapKbps = 0;
 
+	// The amount of bandwidth to reserve during gameplay
+	int32 ReservedGameplayBandwidthKbps = 10000;
+
 	// Incrementing Handle value for external services to track their FManagedService data within the manager
 	int32 NextServiceId = UE_INVALID_BANDWIDTH_SERVICE_HANDLE;
 
