BPS: fix for disk overflow with new directory structure

#rnx
#Synced-CL 51607449

diff --git a/Engine/Source/Runtime/Online/BuildPatchServices/Private/BuildPatchFileConstructor.cpp b/Engine/Source/Runtime/Online/BuildPatchServices/Private/BuildPatchFileConstructor.cpp
--- a/Engine/Source/Runtime/Online/BuildPatchServices/Private/BuildPatchFileConstructor.cpp
+++ b/Engine/Source/Runtime/Online/BuildPatchServices/Private/BuildPatchFileConstructor.cpp
@@ -943,7 +943,7 @@ public:
 
 		if (bUseDiskOverflowStore)
 		{
-			IFileManager::Get().MakeDirectory(*InParentConstructor->Configuration.BackingStoreDirectory);
+			IFileManager::Get().MakeDirectory(*InParentConstructor->Configuration.BackingStoreDirectory, true);
 			BackingStoreFileHandle.Reset(FPlatformFileManager::Get().GetPlatformFile().OpenWrite(*BackingStoreFileName, false, true));
 		}
 
