EventLoop: bug fixes, dead code removal, and expanded test coverage

- Fix GetLoopTime() to assert the caller is the event loop thread (was asserting initialized instead)
- Fix ClearTimer() to skip TryNotify when called with an invalid handle
- Fix repeat timer rescheduling to tick with zero delta, preventing short-interval timers from immediately re-firing
- Remove unused EEventLoopStatus enum (includes Unintialized typo fix), IRequestManager forward decl, and TSharedFromThis base class
- Fix incomplete/garbled RequestShutdown doc comment
- Add 12 missing test cases covering GetLoopTime, IOManager init/shutdown, ClearTimer with valid handle, multiple pending shutdown callbacks, post-shutdown RunOnce behavior, invalid SetTimer params, FirstDelay, multiple simultaneous timers, and cross-callback task posting
- Moved Test build and target cs files to correct location.

#Synced-CL 51619035

diff --git a/Engine/Source/Runtime/Online/EventLoop/Public/EventLoop/EventLoop.h b/Engine/Source/Runtime/Online/EventLoop/Public/EventLoop/EventLoop.h
--- a/Engine/Source/Runtime/Online/EventLoop/Public/EventLoop/EventLoop.h
+++ b/Engine/Source/Runtime/Online/EventLoop/Public/EventLoop/EventLoop.h
@@ -184,8 +184,12 @@ public:
 
 	virtual void ClearTimer(FTimerHandle& InHandle, FOnTimerCleared&& OnTimerCleared = FOnTimerCleared()) override
 	{
+		const bool bWasValid = InHandle.IsValid();
 		TimerManager.ClearTimer(InHandle, MoveTemp(OnTimerCleared));
-		TryNotify();
+		if (bWasValid)
+		{
+			TryNotify();
+		}
 	}
 
 	virtual void PostAsyncTask(FAsyncTask&& Task) override
@@ -239,14 +243,12 @@ public:
 		TimerManager.Tick(ElapsedTime);
 
 		// If any timers ran due to the above call to Tick, check whether any repeat timers are
-		// waiting to be rescheduled. Advancing the loop time and ticking the timers again will
-		// allow those timers to reschedule without requiring polling the IO manager again.
+		// waiting to be rescheduled. Ticking with zero delta reschedules them into the heap
+		// without advancing internal time, which prevents short-interval timers from
+		// immediately re-firing in the same iteration.
 		if (TimerManager.HasPendingRepeatTimer())
 		{
-			CurrentTime = Traits::GetCurrentTime();
-			ElapsedTime = FMath::Max(CurrentTime - LoopTime, 0.0);
-			LoopTime = CurrentTime;
-			TimerManager.Tick(ElapsedTime);
+			TimerManager.Tick(FTimespan::Zero());
 		}
 
 		RunAsyncTasks();
@@ -285,7 +287,7 @@ public:
 
 	virtual FTimespan GetLoopTime() const override
 	{
-		Traits::CheckInitialized(EventLoopThreadId);
+		Traits::CheckIsEventLoopThread(EventLoopThreadId);
 		return LoopTime;
 	}
 
diff --git a/Engine/Source/Runtime/Online/EventLoop/Public/EventLoop/IEventLoop.h b/Engine/Source/Runtime/Online/EventLoop/Public/EventLoop/IEventLoop.h
--- a/Engine/Source/Runtime/Online/EventLoop/Public/EventLoop/IEventLoop.h
+++ b/Engine/Source/Runtime/Online/EventLoop/Public/EventLoop/IEventLoop.h
@@ -10,19 +10,10 @@
 
 namespace UE::EventLoop {
 
-class IRequestManager;
-
-enum class EEventLoopStatus : uint8
-{
-	Unintialized,
-	Running,
-	Terminated,
-};
-
 using FAsyncTask = TUniqueFunction<void()>;
 using FOnShutdownComplete = TUniqueFunction<void()>;
 
-class IEventLoop : public TSharedFromThis<IEventLoop, ESPMode::ThreadSafe>
+class IEventLoop
 {
 public:
 	virtual ~IEventLoop() = default;
@@ -36,9 +27,8 @@ public:
 
 	/**
 	 * Signals a shutdown request to the event loop. May be signaled from any thread.
-	 * Once Shutdown has been called, continue calling Poll until
-	 * Invalidates the timer handle as it should no longer be used.
-	 * 
+	 * Once called, RunOnce will return false after completing its current iteration.
+	 *
 	 * Thread safe.
 	 */
 	virtual void RequestShutdown(FOnShutdownComplete&& OnShutdownComplete = FOnShutdownComplete()) = 0;
diff --git a/Engine/Source/Runtime/Online/EventLoopTests/EventLoopUnitTests.Build.cs b/Engine/Source/Runtime/Online/EventLoopTests/EventLoopUnitTests.Build.cs
new file mode 100644
--- /dev/null
+++ b/Engine/Source/Runtime/Online/EventLoopTests/EventLoopUnitTests.Build.cs
@@ -0,0 +1,24 @@
+// Copyright Epic Games, Inc. All Rights Reserved.
+using UnrealBuildTool;
+
+public class EventLoopUnitTests : TestModuleRules
+{
+    static EventLoopUnitTests()
+    {
+		TestMetadata = new Metadata();
+        TestMetadata.TestName = "EventLoop";
+        TestMetadata.TestShortName = "EventLoop";
+    }
+
+    public EventLoopUnitTests(ReadOnlyTargetRules Target) : base(Target)
+	{
+		OptimizeCode = CodeOptimization.Never;
+
+		PrivateDependencyModuleNames.AddRange(
+			new string[] {
+				"Core",
+				"EventLoop",
+				"Sockets"
+			});
+	}
+}
\ No newline at end of file
diff --git a/Engine/Source/Runtime/Online/EventLoopTests/EventLoopUnitTests.Target.cs b/Engine/Source/Runtime/Online/EventLoopTests/EventLoopUnitTests.Target.cs
new file mode 100644
--- /dev/null
+++ b/Engine/Source/Runtime/Online/EventLoopTests/EventLoopUnitTests.Target.cs
@@ -0,0 +1,15 @@
+// Copyright Epic Games, Inc. All Rights Reserved.
+using UnrealBuildTool;
+
+[SupportedPlatforms(UnrealPlatformClass.All)]
+public class EventLoopUnitTestsTarget : TestTargetRules
+{
+	public EventLoopUnitTestsTarget(TargetInfo Target) : base(Target)
+	{
+		bCompileAgainstCoreUObject = true;
+		bCompileAgainstApplicationCore = true;
+
+		// Optionally add global definitions for Catch2 benchmarking etc.
+		GlobalDefinitions.Add("CATCH_CONFIG_ENABLE_BENCHMARKING=1");
+	}
+}
\ No newline at end of file
diff --git a/Engine/Source/Runtime/Online/EventLoopTests/Tests/EventLoop/EventLoopTests.cpp b/Engine/Source/Runtime/Online/EventLoopTests/Tests/EventLoop/EventLoopTests.cpp
--- a/Engine/Source/Runtime/Online/EventLoopTests/Tests/EventLoop/EventLoopTests.cpp
+++ b/Engine/Source/Runtime/Online/EventLoopTests/Tests/EventLoop/EventLoopTests.cpp
@@ -218,9 +218,9 @@ namespace UE::EventLoop
 
 	TEST_CASE("EventLoop", "[Online][EventLoop][Smoke]")
 	{
-		TEventLoop<FIOManagerMock, FEventLoopMockTraits> EventLoop;
 		FIOManagerMock::ResetTestConditions();
 		FEventLoopMockTraits::ResetTestConditions();
+		TEventLoop<FIOManagerMock, FEventLoopMockTraits> EventLoop;
 
 		const FTimespan InfiniteWait(std::numeric_limits<int64>::max());
 		const uint32 EventLoopThreadId = 1;
@@ -231,6 +231,7 @@ namespace UE::EventLoop
 		{
 			EventLoop.Init();
 			CHECK(!FEventLoopMockTraits::bCheckNotInitializedTriggered);
+			CHECK(FIOManagerMock::bHasInitialized);
 		}
 
 		SECTION("Double init call")
@@ -344,6 +345,41 @@ namespace UE::EventLoop
 			CHECK(bTaskTriggered);
 		}
 
+		SECTION("Init - IOManager failure propagated")
+		{
+			FIOManagerMock::bInitReturnValue = false;
+			CHECK(!EventLoop.Init());
+		}
+
+		SECTION("GetLoopTime - returns time captured at Init")
+		{
+			FEventLoopMockTraits::CurrentTime = FTimespan::FromSeconds(5);
+			EventLoop.Init();
+			CHECK(EventLoop.GetLoopTime() == FTimespan::FromSeconds(5));
+		}
+
+		SECTION("GetLoopTime - updates after each RunOnce")
+		{
+			EventLoop.Init();
+
+			FIOManagerMock::ExpectedWaitTime = InfiniteWait;
+			FEventLoopMockTraits::CurrentTime = FTimespan::FromSeconds(2);
+			EventLoop.RunOnce(InfiniteWait);
+			CHECK(EventLoop.GetLoopTime() == FTimespan::FromSeconds(2));
+
+			FEventLoopMockTraits::CurrentTime = FTimespan::FromSeconds(5);
+			EventLoop.RunOnce(InfiniteWait);
+			CHECK(EventLoop.GetLoopTime() == FTimespan::FromSeconds(5));
+		}
+
+		SECTION("GetLoopTime - fires check when called from wrong thread")
+		{
+			EventLoop.Init(); // EventLoopThreadId set to CurrentThreadId (1).
+			FEventLoopMockTraits::CurrentThreadId = 2;
+			EventLoop.GetLoopTime();
+			CHECK(FEventLoopMockTraits::bCheckIsEventLoopThreadTriggered);
+		}
+
 		SECTION("Initialized")
 		{
 			EventLoop.Init();
@@ -380,6 +416,7 @@ namespace UE::EventLoop
 
 				CHECK(!FEventLoopMockTraits::bCheckInitializedTriggered);
 				CHECK(!FEventLoopMockTraits::bCheckIsEventLoopThreadTriggered);
+				CHECK(FIOManagerMock::bHasShutdown);
 			}
 
 			SECTION("Call SetTimer")
@@ -438,7 +475,7 @@ namespace UE::EventLoop
 			SECTION("Call ClearTimer")
 			{
 				FIOManagerMock::ExpectedWaitTime = InfiniteWait;
-				FIOManagerMock::bExpectedNotify = true;
+				FIOManagerMock::bExpectedNotify = false; // Invalid handle — no notify expected.
 
 				FTimerHandle Handle;
 				EventLoop.ClearTimer(Handle);
@@ -486,6 +523,204 @@ namespace UE::EventLoop
 
 				EventLoop.RunOnce(InfiniteWait);
 			}
+
+			SECTION("Call RunOnce - returns false without polling after shutdown")
+			{
+				// Drive the loop to shutdown.
+				FEventLoopMockTraits::bShutdownRequested = true;
+				FIOManagerMock::ExpectedWaitTime = InfiniteWait;
+				CHECK(!EventLoop.RunOnce(InfiniteWait));
+
+				// A subsequent RunOnce must return false immediately without calling Poll or
+				// Shutdown again. If Poll is incorrectly called it will receive a WaitTime that
+				// mismatches ExpectedWaitTime, causing the internal CHECK inside Poll to fire.
+				FIOManagerMock::bHasShutdown = false;
+				FIOManagerMock::ExpectedWaitTime = FTimespan::Zero();
+				CHECK(!EventLoop.RunOnce(FTimespan::FromSeconds(999)));
+				CHECK(!FIOManagerMock::bHasShutdown);
+			}
+
+			SECTION("Call SetTimer - invalid parameters return invalid handle")
+			{
+				// None of these should produce a valid timer. TryNotify is still called
+				// by TEventLoop::SetTimer regardless, but we do not RunOnce so Poll is
+				// never reached and the bHasNotified state does not matter.
+				FTimerHandle Handle;
+
+				// Null callback.
+				Handle = EventLoop.SetTimer(FTimerCallback(), FTimespan::FromSeconds(1));
+				CHECK(!Handle.IsValid());
+
+				// Negative rate.
+				Handle = EventLoop.SetTimer([](){ }, FTimespan::FromSeconds(-1));
+				CHECK(!Handle.IsValid());
+
+				// FirstDelay on a non-repeating timer.
+				Handle = EventLoop.SetTimer([](){ }, FTimespan::FromSeconds(1), false, FTimespan::FromSeconds(1));
+				CHECK(!Handle.IsValid());
+
+				// Negative FirstDelay.
+				Handle = EventLoop.SetTimer([](){ }, FTimespan::FromSeconds(1), true, FTimespan::FromSeconds(-1));
+				CHECK(!Handle.IsValid());
+			}
+
+			SECTION("Call SetTimer - repeating with FirstDelay")
+			{
+				bool bTimerRan = false;
+
+				// Activate the timer. The heap is empty at this point so Poll waits infinitely.
+				FIOManagerMock::bExpectedNotify = true;
+				FIOManagerMock::ExpectedWaitTime = InfiniteWait;
+				EventLoop.SetTimer([&bTimerRan](){ bTimerRan = true; }, FTimespan::FromSeconds(1), true, FTimespan::FromSeconds(2));
+				EventLoop.RunOnce(InfiniteWait); // activates timer with Expiration = 0 + FirstDelay(2s)
+				CHECK(!bTimerRan);
+
+				// Poll wait is clamped to FirstDelay (2s), not Rate (1s).
+				FIOManagerMock::bHasNotified = false;
+				FIOManagerMock::bExpectedNotify = false;
+				FIOManagerMock::ExpectedWaitTime = FTimespan::FromSeconds(2);
+				FEventLoopMockTraits::CurrentTime = FTimespan::FromSeconds(1); // Rate elapsed, but not FirstDelay.
+				EventLoop.RunOnce(InfiniteWait);
+				CHECK(!bTimerRan);
+
+				// FirstDelay elapsed — timer fires and reschedules at Rate.
+				// Poll wait = Expiration(2s) - InternalTime(1s) = 1s.
+				FIOManagerMock::ExpectedWaitTime = FTimespan::FromSeconds(1);
+				FEventLoopMockTraits::CurrentTime = FTimespan::FromSeconds(2);
+				EventLoop.RunOnce(InfiniteWait);
+				CHECK(bTimerRan);
+
+				// Timer repeats at Rate going forward.
+				bTimerRan = false;
+				FIOManagerMock::ExpectedWaitTime = FTimespan::FromSeconds(1);
+				FEventLoopMockTraits::CurrentTime = FTimespan::FromSeconds(3);
+				EventLoop.RunOnce(InfiniteWait);
+				CHECK(bTimerRan);
+			}
+
+			SECTION("Call SetTimer - multiple timers fire in order")
+			{
+				bool bFirstTimerRan = false;
+				bool bSecondTimerRan = false;
+
+				// Both SetTimer calls notify; bHasNotified ends up true either way.
+				FIOManagerMock::bExpectedNotify = true;
+				FIOManagerMock::ExpectedWaitTime = InfiniteWait;
+				EventLoop.SetTimer([&bFirstTimerRan](){ bFirstTimerRan = true; }, FTimespan::FromSeconds(1));
+				EventLoop.SetTimer([&bSecondTimerRan](){ bSecondTimerRan = true; }, FTimespan::FromSeconds(2));
+				EventLoop.RunOnce(InfiniteWait); // activates both timers
+				CHECK(!bFirstTimerRan);
+				CHECK(!bSecondTimerRan);
+
+				// Poll is clamped to the nearest expiration (1s).
+				FIOManagerMock::bHasNotified = false;
+				FIOManagerMock::bExpectedNotify = false;
+				FIOManagerMock::ExpectedWaitTime = FTimespan::FromSeconds(1);
+				FEventLoopMockTraits::CurrentTime = FTimespan::FromSeconds(1);
+				EventLoop.RunOnce(InfiniteWait);
+				CHECK(bFirstTimerRan);
+				CHECK(!bSecondTimerRan);
+
+				// Poll clamped to remaining time until second timer:
+				// Expiration(2s) - InternalTime(1s) = 1s.
+				FIOManagerMock::ExpectedWaitTime = FTimespan::FromSeconds(1);
+				FEventLoopMockTraits::CurrentTime = FTimespan::FromSeconds(2);
+				EventLoop.RunOnce(InfiniteWait);
+				CHECK(bSecondTimerRan);
+			}
+
+			SECTION("Call ClearTimer - valid handle")
+			{
+				bool bTimerRan = false;
+				bool bTimerCleared = false;
+
+				// Activate the timer.
+				FIOManagerMock::bExpectedNotify = true;
+				FIOManagerMock::ExpectedWaitTime = InfiniteWait;
+				FTimerHandle Handle = EventLoop.SetTimer([&bTimerRan](){ bTimerRan = true; }, FTimespan::FromSeconds(1));
+				CHECK(Handle.IsValid());
+				EventLoop.RunOnce(InfiniteWait); // timer is now active in the heap
+
+				// Clear the valid handle — should notify. The timer is still in the heap
+				// during GetNextTimeout so Poll is clamped to its remaining time (1s).
+				FIOManagerMock::bHasNotified = false;
+				FIOManagerMock::bExpectedNotify = true;
+				FIOManagerMock::ExpectedWaitTime = FTimespan::FromSeconds(1);
+				EventLoop.ClearTimer(Handle, [&bTimerCleared](){ bTimerCleared = true; });
+				CHECK(!Handle.IsValid());
+				EventLoop.RunOnce(InfiniteWait); // processes removal and fires OnTimerCleared
+				CHECK(bTimerCleared);
+				CHECK(!bTimerRan);
+
+				// Advance past the original expiration — the cleared timer must not fire.
+				FIOManagerMock::bHasNotified = false;
+				FIOManagerMock::bExpectedNotify = false;
+				FIOManagerMock::ExpectedWaitTime = InfiniteWait; // heap is now empty
+				FEventLoopMockTraits::CurrentTime = FTimespan::FromSeconds(1);
+				EventLoop.RunOnce(InfiniteWait);
+				CHECK(!bTimerRan);
+			}
+
+			SECTION("Call RequestShutdown - multiple callbacks queued before RunOnce")
+			{
+				bool bFirstShutdownCompleted = false;
+				bool bSecondShutdownCompleted = false;
+
+				// Only the first RequestShutdown calls TryNotify; the second queues its
+				// callback without notifying since shutdown is already requested.
+				FIOManagerMock::bExpectedNotify = true;
+				FIOManagerMock::ExpectedWaitTime = InfiniteWait;
+				EventLoop.RequestShutdown([&bFirstShutdownCompleted](){ bFirstShutdownCompleted = true; });
+				EventLoop.RequestShutdown([&bSecondShutdownCompleted](){ bSecondShutdownCompleted = true; });
+
+				bool bContinueRunning = EventLoop.RunOnce(InfiniteWait);
+				CHECK(!bContinueRunning);
+				CHECK(bFirstShutdownCompleted);
+				CHECK(bSecondShutdownCompleted);
+			}
+
+			SECTION("Call PostAsyncTask - task posts another task in same iteration")
+			{
+				bool bFirstTaskRan = false;
+				bool bSecondTaskRan = false;
+
+				// The second task is enqueued while RunAsyncTasks is draining the queue.
+				// The while loop picks it up before returning, so both run in one RunOnce.
+				FIOManagerMock::bExpectedNotify = true;
+				FIOManagerMock::ExpectedWaitTime = InfiniteWait;
+				EventLoop.PostAsyncTask([&](){
+					bFirstTaskRan = true;
+					EventLoop.PostAsyncTask([&bSecondTaskRan](){ bSecondTaskRan = true; });
+				});
+				EventLoop.RunOnce(InfiniteWait);
+				CHECK(bFirstTaskRan);
+				CHECK(bSecondTaskRan);
+			}
+
+			SECTION("Timer callback posts async task in same RunOnce iteration")
+			{
+				bool bTimerRan = false;
+				bool bTaskRan = false;
+
+				// Activate the timer.
+				FIOManagerMock::bExpectedNotify = true;
+				FIOManagerMock::ExpectedWaitTime = InfiniteWait;
+				EventLoop.SetTimer([&](){
+					bTimerRan = true;
+					EventLoop.PostAsyncTask([&bTaskRan](){ bTaskRan = true; });
+				}, FTimespan::FromSeconds(1));
+				EventLoop.RunOnce(InfiniteWait);
+
+				// Timer fires during Tick; callback posts a task. RunAsyncTasks runs
+				// immediately after Tick in the same RunOnce, so the task executes here too.
+				FIOManagerMock::bHasNotified = false;
+				FIOManagerMock::bExpectedNotify = false;
+				FIOManagerMock::ExpectedWaitTime = FTimespan::FromSeconds(1);
+				FEventLoopMockTraits::CurrentTime = FTimespan::FromSeconds(1);
+				EventLoop.RunOnce(InfiniteWait);
+				CHECK(bTimerRan);
+				CHECK(bTaskRan);
+			}
 		}
 	}
 } // UE::EventLoop
diff --git a/Engine/Source/Runtime/Online/Experimental/EventLoopTests/EventLoopUnitTests.Build.cs b/Engine/Source/Runtime/Online/Experimental/EventLoopTests/EventLoopUnitTests.Build.cs
deleted file mode 100644
diff --git a/Engine/Source/Runtime/Online/Experimental/EventLoopTests/EventLoopUnitTests.Target.cs b/Engine/Source/Runtime/Online/Experimental/EventLoopTests/EventLoopUnitTests.Target.cs
deleted file mode 100644
