Runtime CVar to control CURLMOPT_MAX_HOST_CONNECTIONS on the HTTP thread

Added HTTP.Curl.MaxHostConnections CVar to change the libcurl per-host connection limit (CURLMOPT_MAX_HOST_CONNECTIONS) at runtime. This is useful when game code knows a burst of concurrent HTTP requests to the same host is about to happen and wants to avoid the default per-host cap from serialising them.

The CVar is seeded at init from the existing [HTTP] HttpMaxConnectionsPerServer config value. A change delegate registered in InitCurl() dispatches curl_multi_setopt to the HTTP thread to modify CURLMOPT_MAX_HOST_CONNECTIONS.

#Synced-CL 51753475

diff --git a/Engine/Source/Runtime/Online/HTTP/Private/Curl/CurlHttpManager.cpp b/Engine/Source/Runtime/Online/HTTP/Private/Curl/CurlHttpManager.cpp
--- a/Engine/Source/Runtime/Online/HTTP/Private/Curl/CurlHttpManager.cpp
+++ b/Engine/Source/Runtime/Online/HTTP/Private/Curl/CurlHttpManager.cpp
@@ -90,11 +90,43 @@ TAutoConsoleVariable<bool> CVarHttpCurlReuseConnectionEnabled(
 	ECVF_SaveForNextBoot
 );
 
+// Seeded once during InitCurl() from the [HTTP] HttpMaxConnectionsPerServer config value if > 0.
+TAutoConsoleVariable<int32> CVarHttpCurlMaxHostConnections(
+	TEXT("HTTP.Curl.MaxHostConnections"),
+	0,
+	TEXT("The max number of connections per host (CURLMOPT_MAX_HOST_CONNECTIONS). 0 means unlimited. Can be changed at runtime."),
+	ECVF_Default
+);
+
 namespace UE::Curl
 {
 
 static FCurlInitializeData InitializeData;
 
+/** Set the max number of connections per host at runtime. Must be called from the HTTP thread. */
+static void SetMaxHostConnections(int32 InMaxHostConnections)
+{
+	if (!FCurlHttpManager::IsInit())
+	{
+		UE_LOGF(LogHttp, Warning, "SetMaxHostConnections: curl not initialized, cannot change limit");
+		return;
+	}
+
+	// Normalize negative values to 0 (unlimited)
+	const int32 EffectiveMaxHostConnections = FMath::Max(InMaxHostConnections, 0);
+
+	const CURLMcode SetOptResult = curl_multi_setopt(FCurlHttpManager::GMultiHandle, CURLMOPT_MAX_HOST_CONNECTIONS, static_cast<long>(EffectiveMaxHostConnections));
+	if (SetOptResult != CURLM_OK)
+	{
+		FUTF8ToTCHAR Converter(curl_multi_strerror(SetOptResult));
+		UE_LOGF(LogHttp, Warning, "SetMaxHostConnections: Failed to set to %d, error %d ('%ls')", EffectiveMaxHostConnections, (int32)SetOptResult, Converter.Get());
+		return;
+	}
+
+	UE_LOGF(LogHttp, Log, "SetMaxHostConnections: Changed from %d to %d", FCurlHttpManager::CurlRequestOptions.MaxHostConnections, EffectiveMaxHostConnections);
+	FCurlHttpManager::CurlRequestOptions.MaxHostConnections = EffectiveMaxHostConnections;
+}
+
 }
 
 bool FCurlHttpManager::IsInit()
@@ -257,7 +289,15 @@ void FCurlHttpManager::InitCurl()
 
 	GConfig->GetBool(TEXT("HTTP.Curl"), TEXT("bAllowSeekFunction"), CurlRequestOptions.bAllowSeekFunction, GEngineIni);
 
-	CurlRequestOptions.MaxHostConnections = FHttpModule::Get().GetHttpMaxConnectionsPerServer();
+
+	// Seed the CVar from the [HTTP] config value
+	const int32 ConfigMaxHostConnections = FHttpModule::Get().GetHttpMaxConnectionsPerServer();
+	if (ConfigMaxHostConnections > 0)
+	{
+		CVarHttpCurlMaxHostConnections.AsVariable()->Set(ConfigMaxHostConnections, ECVF_SetByGameSetting);
+	}
+	CurlRequestOptions.MaxHostConnections = FMath::Max(CVarHttpCurlMaxHostConnections.GetValueOnAnyThread(), 0);
+
 	if (CurlRequestOptions.MaxHostConnections > 0)
 	{
 		const CURLMcode SetOptResult = curl_multi_setopt(GMultiHandle, CURLMOPT_MAX_HOST_CONNECTIONS, static_cast<long>(CurlRequestOptions.MaxHostConnections));
@@ -274,6 +314,7 @@ void FCurlHttpManager::InitCurl()
 		CurlRequestOptions.MaxHostConnections = 0;
 	}
 
+
 	TCHAR Home[256] = TEXT("");
 	if (FParse::Value(FCommandLine::Get(), TEXT("MULTIHOMEHTTP="), Home, UE_ARRAY_COUNT(Home)))
 	{
@@ -286,6 +327,20 @@ void FCurlHttpManager::InitCurl()
 
 	// print for visibility
 	CurlRequestOptions.Log();
+
+	// Bridge game-thread CVar changes to the HTTP thread. When game code calls Set() on this CVar,
+	// the delegate captures the new value and dispatches it via AddHttpThreadTask so that
+	// curl_multi_setopt is called on the correct thread.
+	CVarHttpCurlMaxHostConnections.AsVariable()->SetOnChangedCallback(
+		FConsoleVariableDelegate::CreateLambda([](IConsoleVariable* Variable)
+		{
+			const int32 NewValue = Variable->GetInt();
+			FHttpModule::Get().GetHttpManager().AddHttpThreadTask([NewValue]()
+			{
+				UE::Curl::SetMaxHostConnections(NewValue);
+			});
+		})
+	);
 }
 
 void FCurlHttpManager::FCurlRequestOptions::Log()
@@ -334,6 +389,8 @@ void FCurlHttpManager::ShutdownCurl()
 	}
 #endif
 
+	CVarHttpCurlMaxHostConnections.AsVariable()->SetOnChangedCallback(FConsoleVariableDelegate());
+
 	if (GMultiHandle != nullptr)
 	{
 		CURLMcode MutliCleanupCode = curl_multi_cleanup(GMultiHandle);
