[Http]Trigger fatal log when Flush during http complete callback since it's impossible.

#Synced-CL 52039227

diff --git a/Engine/Source/Runtime/Online/HTTP/Private/HttpManager.cpp b/Engine/Source/Runtime/Online/HTTP/Private/HttpManager.cpp
--- a/Engine/Source/Runtime/Online/HTTP/Private/HttpManager.cpp
+++ b/Engine/Source/Runtime/Online/HTTP/Private/HttpManager.cpp
@@ -12,6 +12,7 @@
 #include "Misc/Fork.h"
 #include "HttpThread.h"
 #include "GenericPlatform/HttpRequestCommon.h"
+#include "Logging/StructuredLog.h"
 #include "Misc/ConfigCacheIni.h"
 #include "Misc/CommandLine.h"
 #include "ProfilingDebugging/CsvProfiler.h"
@@ -53,6 +54,12 @@ TAutoConsoleVariable<bool> CVarHttpRemoveRequestUsingHttpThreadPolicyOnHttpThrea
 	ECVF_SaveForNextBoot
 );
 
+TAutoConsoleVariable<bool> CVarHttpFlushTriggerFatalLogDuringCompleteDelegateEnabled(
+	TEXT("http.FlushTriggerFatalLogDuringCompleteDelegateEnabled"),
+	true,
+	TEXT("When enabled, trigger fatal log if flush during http complete delegate")
+);
+
 // FHttpManager
 
 const TCHAR* LexToString(const EHttpFlushReason& FlushReason)
@@ -372,7 +379,12 @@ void FHttpManager::Flush(EHttpFlushReason FlushReason)
 {
 	QUICK_SCOPE_CYCLE_COUNTER(STAT_FHttpManager_Flush);
 
-	UE_CLOG(IsCurrentThreadCompletingRequest(), LogHttp, Warning, TEXT("It's not supported to call Flush during http request complete delegate!"));
+	if (IsCurrentThreadCompletingRequest())
+	{
+		UE_CLOGFMT(CVarHttpFlushTriggerFatalLogDuringCompleteDelegateEnabled.GetValueOnAnyThread(), LogHttp, Fatal, "It's not supported to call Flush during http request complete delegate!");
+		return;
+	}
+
 	checkf(FlushReason != EHttpFlushReason::Shutdown || !HasAnyBoundDelegate(), TEXT("Use Shutdown() instead of Flush(EHttpFlushReason::Shutdown) directly."));
 	
 	// This variable is set to indicate that flush is happening.
