[Http]Don't create http module in FHttpModule::Get(), to make it easier to see the issue of accessing HttpModule::Get after shutting down.

#Synced-CL 52039227

#Synced-CL 52039227

diff --git a/Engine/Source/Runtime/Online/HTTP/Private/HttpModule.cpp b/Engine/Source/Runtime/Online/HTTP/Private/HttpModule.cpp
--- a/Engine/Source/Runtime/Online/HTTP/Private/HttpModule.cpp
+++ b/Engine/Source/Runtime/Online/HTTP/Private/HttpModule.cpp
@@ -34,6 +34,13 @@ TAutoConsoleVariable<bool> CVarHttpInsecureProtocolEnabled(
 );
 #endif
 
+TAutoConsoleVariable<bool> CVarHttpAutoCreateModuleEnabled(
+	TEXT("http.AutoCreateModuleEnabled"),
+	false,
+	TEXT("When enabled, auto create HttpModule when call HttpModule::Get(). NOTE: Enabling this can make it hard to trace the case of recreating http module after shutting down."),
+	ECVF_SaveForNextBoot
+);
+
 DEFINE_LOG_CATEGORY(LogHttp);
 
 // FHttpModule
@@ -367,7 +374,7 @@ bool FHttpModule::Exec_Runtime(UWorld* InWorld, const TCHAR* Cmd, FOutputDevice&
 
 FHttpModule& FHttpModule::Get()
 {
-	if (Singleton == nullptr)
+	if (Singleton == nullptr && CVarHttpAutoCreateModuleEnabled.GetValueOnAnyThread())
 	{
 		check(IsInGameThread());
 		FModuleManager::LoadModuleChecked<FHttpModule>("HTTP");
