Initial Bandwidth Measurement Tool

#rnx

#Synced-CL 51749326

diff --git a/Engine/Source/Runtime/Engine/Engine.Build.cs b/Engine/Source/Runtime/Engine/Engine.Build.cs
--- a/Engine/Source/Runtime/Engine/Engine.Build.cs
+++ b/Engine/Source/Runtime/Engine/Engine.Build.cs
@@ -153,6 +153,7 @@ public class Engine : ModuleRules
 				"AnimationCore",
 				"AutoRTFM",
 				"AppFramework",
+				"BandwidthMeasurementTool",
 				"BuildSettings",
 				"Networking",
 				"Landscape",
diff --git a/Engine/Source/Runtime/Launch/Private/LaunchEngineLoop.cpp b/Engine/Source/Runtime/Launch/Private/LaunchEngineLoop.cpp
--- a/Engine/Source/Runtime/Launch/Private/LaunchEngineLoop.cpp
+++ b/Engine/Source/Runtime/Launch/Private/LaunchEngineLoop.cpp
@@ -2915,6 +2915,11 @@ int32 FEngineLoop::PreInitPreStartupScreen(const TCHAR* CmdLine)
 	FModuleManager::Get().LoadModule("IoStoreOnDemand");
 #endif // UE_WITH_IOSTOREONDEMAND
 
+#if UE_CLIENT || UE_EDITOR
+	FModuleManager::Get().LoadModule("BandwidthMeasurementTool");
+	FModuleManager::Get().LoadModule("DownlinkBandwidthManager");
+#endif // UE_CLIENT || UE_EDITOR
+
 #endif // USE_IO_DISPATCHER
 
 	// Let LogConsole know what ini file it should use to save its setting on exit.
diff --git a/Engine/Source/Runtime/Online/Experimental/BandwidthMeasurementTool/BandwidthMeasurementTool.Build.cs b/Engine/Source/Runtime/Online/Experimental/BandwidthMeasurementTool/BandwidthMeasurementTool.Build.cs
new file mode 100644
--- /dev/null
+++ b/Engine/Source/Runtime/Online/Experimental/BandwidthMeasurementTool/BandwidthMeasurementTool.Build.cs
@@ -0,0 +1,21 @@
+// Copyright Epic Games, Inc. All Rights Reserved.
+
+using UnrealBuildTool;
+
+public class BandwidthMeasurementTool : ModuleRules
+{
+	public BandwidthMeasurementTool(ReadOnlyTargetRules Target) : base(Target)
+	{
+		PCHUsage = PCHUsageMode.UseExplicitOrSharedPCHs;
+		CppCompileWarningSettings.UnsafeTypeCastWarningLevel = WarningLevel.Error;
+
+		PublicDependencyModuleNames.AddRange(new string[]{
+			"Core",
+			});
+
+		PrivateDependencyModuleNames.AddRange(new string[]{
+			"CoreUObject",
+			"HTTP",
+			});
+	}
+}
\ No newline at end of file
diff --git a/Engine/Source/Runtime/Online/Experimental/BandwidthMeasurementTool/Private/BandwidthMeasurementTool.cpp b/Engine/Source/Runtime/Online/Experimental/BandwidthMeasurementTool/Private/BandwidthMeasurementTool.cpp
new file mode 100644
--- /dev/null
+++ b/Engine/Source/Runtime/Online/Experimental/BandwidthMeasurementTool/Private/BandwidthMeasurementTool.cpp
@@ -0,0 +1,387 @@
+// Copyright Epic Games, Inc. All Rights Reserved.
+
+#include "BandwidthMeasurementTool.h"
+#include "HAL/IConsoleManager.h"
+#include "HAL/PlatformTime.h"
+#include "HttpModule.h"
+#include "Interfaces/IHttpRequest.h"
+#include "Interfaces/IHttpResponse.h"
+#include "Math/NumericLimits.h"
+#include "Misc/ConfigCacheIni.h"
+#include "Templates/UniquePtr.h"
+
+DEFINE_LOG_CATEGORY(LogBandwidthMeasurementTool);
+
+namespace BWMT
+{
+	static bool GEnableMeasurementTool = false;
+	static FAutoConsoleVariableRef CVarBWMTEnableMeasurementTool(
+		TEXT("net.BWMT.Enable"),
+		GEnableMeasurementTool,
+		TEXT("Turns on/off the Bandwidth Measurement Tool."),
+		ECVF_Default);
+
+	static int32 GMaxThroughputHistory = 10;
+	static FAutoConsoleVariableRef CVarBWMTMaxThroughputHistory(
+		TEXT("net.BWMT.HistoryCount"),
+		GMaxThroughputHistory,
+		TEXT("Value used to denote the maximum number of previous measurements the tool can average against."),
+		ECVF_Default);
+
+	static int32 GMaxNumHttpRequests = 5;
+	static FAutoConsoleVariableRef CVarBWMTMaxNumHttpRequests(
+		TEXT("net.BWMT.RequestCount"),
+		GMaxNumHttpRequests,
+		TEXT("The number of network requests to send out to flood the network for measurement."),
+		ECVF_Default);
+
+	static float GMaxTimeInSecondsToMeasure = 5.0f;
+	static FAutoConsoleVariableRef CVarBWMTMaxTimeInSecondsToMeasure(
+		TEXT("net.BWMT.MeasurementTime"),
+		GMaxTimeInSecondsToMeasure,
+		TEXT("The time in seconds that we want to allow the measurement tool to collect data from the download requests."),
+		ECVF_Default);
+
+	static float GWarmupPeriodTimeInSeconds = 0.0f; 
+	static FAutoConsoleVariableRef CVarBWMTWarmupPeriodTimeInSeconds(
+		TEXT("net.BWMT.WarmupPeriod"),
+		GWarmupPeriodTimeInSeconds,
+		TEXT("The time in seconds we want to not capture download data to let the CDN warmup and provide a more accurate measurement."),
+		ECVF_Default);
+
+	static FString GMeasurementURL = TEXT("");
+	static FAutoConsoleVariableRef CVarBWMTMeasurementURL(
+		TEXT("net.BWMT.MeasurementURL"),
+		GMeasurementURL,
+		TEXT("The CDN URL to download from when measuring bandwidth."),
+		ECVF_Default);
+
+	// Command to manually run a measurement and log the outcome
+	FAutoConsoleCommand FManualBandwidthMeasurementTest(
+		TEXT("net.BWMT.ManualMeasure"), 
+		TEXT("Starts a manual run of the bandwidth measurement tool and will log the resulting value."),
+		FConsoleCommandDelegate::CreateLambda([]() 
+			{
+				FBandwidthMeasurementTool& MeasurementTool = FBandwidthMeasurementTool::Get();
+				MeasurementTool.RegisterForCompletedMeasurement([](uint64 MeasurementInBps) 
+					{
+						UE_LOGF(LogBandwidthMeasurementTool, Log, "FBandwidthMeasurementTool has measured a bandwidth of %llu Bps.", MeasurementInBps);
+					});
+
+				MeasurementTool.BeginMeasurement();
+			}),
+		ECVF_Cheat);
+
+	// Command to manually clear the accumulated historical measurement data
+	FAutoConsoleCommand FClearMeasurementHistory(
+		TEXT("net.BWMT.ClearHistory"),
+		TEXT("Clears the collected history of bandwidth measurements."),
+		FConsoleCommandDelegate::CreateLambda([]()
+			{
+				FBandwidthMeasurementTool& MeasurementTool = FBandwidthMeasurementTool::Get();
+				MeasurementTool.ClearHistory();
+			}),
+		ECVF_Cheat);
+
+	// Command to print telemetry data to the log
+	FAutoConsoleCommand FPrintTelemetryData(
+		TEXT("net.BWMT.PrintTelemetryData"),
+		TEXT("Prints most recent snapshot of bandwidth measurement data."),
+		FConsoleCommandDelegate::CreateLambda([]()
+			{
+				FBandwidthMeasurementTool& MeasurementTool = FBandwidthMeasurementTool::Get();
+				const FBandwidthMeasurementTool::FBandwidthSnapshotData& Telemetry = MeasurementTool.GetTelemetryData();
+
+				UE_LOGF(LogBandwidthMeasurementTool, Log, "Bandwidth Measurement - Finished Requests: %llu\t\tCompleted Requests: %llu\t\tPost-Warmup Completed: %llu",
+					Telemetry.FinishRequestCount, Telemetry.CompleteRequestCount, Telemetry.PostWarmupCompleteRequestCount);
+
+				UE_LOGF(LogBandwidthMeasurementTool, Log, "Bandwidth Measurement - Final Received Bytes: %llu\t\tRaw Bps: %llu",
+					Telemetry.RecentReportedBytesReceived, Telemetry.RecentMeasurementInBps);
+
+				UE_LOGF(LogBandwidthMeasurementTool, Log, "Bandwidth Measurement - Measured Time: %f\t\tMeasurement Count: %llu",
+					Telemetry.RecentMeasurementTime, Telemetry.RecentMeasurementCount);
+
+				UE_LOGF(LogBandwidthMeasurementTool, Log, "Bandwidth Measurement - Recent Historical Measurements:");
+				for (const uint64& Measurement : Telemetry.RecentHistory)
+				{
+					UE_LOGF(LogBandwidthMeasurementTool, Log, "Previous Bps: %llu", Measurement);
+				}
+
+				UE_LOGF(LogBandwidthMeasurementTool, Log, "Bandwidth Measurement - Historical Avg Bps: %llu", Telemetry.RecentAverageInBps);
+			}),
+		ECVF_Cheat);
+}
+
+// ---------------------------
+// Singleton backing pointer
+// ---------------------------
+TUniquePtr<FBandwidthMeasurementTool> FBandwidthMeasurementTool::Instance = nullptr;
+
+/// Static Functions
+FBandwidthMeasurementTool& FBandwidthMeasurementTool::Get()
+{
+	if (!Instance)
+	{
+		Instance = MakeUnique<FBandwidthMeasurementTool>();
+	}
+	check(Instance);
+	return *Instance.Get();
+}
+
+/// Member Functions
+FBandwidthMeasurementTool::FBandwidthMeasurementTool()
+{
+
+}
+
+FBandwidthMeasurementTool::~FBandwidthMeasurementTool()
+{
+
+}
+
+void FBandwidthMeasurementTool::Initialize()
+{
+	FString MeasurementUrl;
+	GConfig->GetString(TEXT("BandwidthMeasurementTool"), TEXT("MeasurementUrl"), MeasurementUrl, GEngineIni);
+	BWMT::CVarBWMTMeasurementURL->Set(*MeasurementUrl, ECVF_SetByCode);
+}
+
+void FBandwidthMeasurementTool::Shutdown()
+{
+	OnBandwidthMeasurementComplete.Clear();
+}
+
+bool FBandwidthMeasurementTool::IsBandwidthBeingMeasured() const
+{
+	return MeasurementTimer.IsValid();
+}
+
+void FBandwidthMeasurementTool::RegisterForCompletedMeasurement(TFunction<void(uint64)>&& FunctionPtr)
+{
+	if (FunctionPtr.IsSet())
+	{
+		OnBandwidthMeasurementComplete.AddLambda(MoveTemp(FunctionPtr));
+	}
+}
+
+float FBandwidthMeasurementTool::GetTimeNeededForMeasurement() const
+{
+	return BWMT::GMaxTimeInSecondsToMeasure;
+}
+
+void FBandwidthMeasurementTool::BeginMeasurement()
+{
+	if (IsBandwidthBeingMeasured() || !BWMT::GEnableMeasurementTool)
+	{
+		// We already have a measurement being performed
+		return;
+	}
+
+	CreateRequestTraffic();
+
+	// We only want this tick to finish when we meet our requirements for stopping
+	MeasurementTimer = FTSTicker::GetCoreTicker().AddTicker(FTickerDelegate::CreateLambda([this](float InDeltaSeconds)->bool
+		{
+			UE::TScopeLock Lock(Mutex);
+
+			bool bTimeToStop = bHasTheLastRequestFinishedEarly;
+			if (!bTimeToStop && MeasurementStartTimestamp > 0.0)
+			{
+				const double ElapsedMeasurementTime = FPlatformTime::Seconds() - MeasurementStartTimestamp;
+				bTimeToStop = ElapsedMeasurementTime >= GetTimeNeededForMeasurement();
+			}
+
+			if (bTimeToStop)
+			{
+				StopMeasuring();
+			}
+
+			return !bTimeToStop;
+		}));
+
+	// This will create the delay time we want to allow the CDN to warm up
+	WarmupTimer = FTSTicker::GetCoreTicker().AddTicker(FTickerDelegate::CreateLambda([this](float InDeltaSeconds)->bool
+		{
+			UE::TScopeLock Lock(Mutex);
+
+			bReadyToMeasureAfterWarmupPeriod = true;
+			MeasurementStartTimestamp = FPlatformTime::Seconds();
+			return false;
+		}), BWMT::GWarmupPeriodTimeInSeconds);
+}
+
+uint64 FBandwidthMeasurementTool::SafeIncrementForUInt64(const uint64 Base, const uint64 Increase) const
+{
+	if (Base <= TNumericLimits<uint64>::Max() - Increase)
+	{
+		// adding is within or at numerical max
+		return Base + Increase;
+	}
+	else
+	{
+		// adding the average would take us outside the max, so we are capping
+		return TNumericLimits<uint64>::Max();
+	}
+}
+
+void FBandwidthMeasurementTool::CreateRequestTraffic()
+{
+	UE::TScopeLock Lock(Mutex);
+
+	TelemetryData.FinishRequestCount = 0;
+	TelemetryData.CompleteRequestCount = 0;
+	TelemetryData.PostWarmupCompleteRequestCount = 0;
+
+	// Need to repeatedly send requests to fill line
+	int32 LastFailedRequest = 0;
+	constexpr int32 SpamPreventingRange = 0;
+	for (int32 RequestCount = 0; RequestCount < BWMT::GMaxNumHttpRequests; RequestCount++)
+	{
+		TSharedRef<IHttpRequest, ESPMode::ThreadSafe> HttpRequest = FHttpModule::Get().CreateRequest();
+		HttpRequest->SetVerb(TEXT("GET"));
+		HttpRequest->SetURL(BWMT::GMeasurementURL);
+		HttpRequest->SetTimeout(GetTimeNeededForMeasurement());
+		HttpRequest->SetResponseBodyReceiveStreamDelegateV2(FHttpRequestStreamDelegateV2::CreateLambda([this](void* Ptr, int64& Length)
+			{
+				UE::TScopeLock Lock(Mutex);
+
+				if (bReadyToMeasureAfterWarmupPeriod && IsBandwidthBeingMeasured())
+				{
+					TotalBytesReceived = SafeIncrementForUInt64(TotalBytesReceived, Length);
+
+					UE_LOGF(LogBandwidthMeasurementTool, Verbose, "Packet Data received! BytesReceived: %llu \t CurrentTotal: %llu", Length, TotalBytesReceived.load());
+				}
+			}));
+
+		HttpRequest->OnProcessRequestComplete().BindLambda(
+			[this](FHttpRequestPtr Request, FHttpResponsePtr Response, bool bProcessedSuccessfully)
+			{
+				UE::TScopeLock Lock(Mutex);
+
+				if (IsBandwidthBeingMeasured())
+				{
+					TelemetryData.FinishRequestCount++;
+
+					if (bProcessedSuccessfully)
+					{
+						TelemetryData.CompleteRequestCount++;
+
+						if (bReadyToMeasureAfterWarmupPeriod)
+						{
+							TelemetryData.PostWarmupCompleteRequestCount++;
+						}
+					}
+
+					if (TelemetryData.FinishRequestCount >= static_cast<uint64>(ActiveRequests.Num()))
+					{
+						bHasTheLastRequestFinishedEarly = true;
+					}
+				}
+			});
+
+		if (HttpRequest->ProcessRequest())
+		{
+			ActiveRequests.Add(HttpRequest);
+		}
+		else
+		{
+			if (RequestCount >= LastFailedRequest + SpamPreventingRange)
+			{
+				// Log a failed processing attempt
+				UE_LOGF(LogBandwidthMeasurementTool, Warning, "%ls - Request attempt %i Failed to process!",
+					UE_SOURCE_LOCATION, RequestCount);
+
+				LastFailedRequest = RequestCount;
+			}
+		}
+	}
+}
+
+void FBandwidthMeasurementTool::StopMeasuring()
+{
+	UE::TScopeLock Lock(Mutex);
+
+	// We have had the data collection period and need to remove the timers to prevent unnecessary execution
+	FTSTicker::GetCoreTicker().RemoveTicker(MeasurementTimer);
+	MeasurementTimer.Reset();
+	FTSTicker::GetCoreTicker().RemoveTicker(WarmupTimer);
+	WarmupTimer.Reset();
+
+	for (TSharedRef<IHttpRequest, ESPMode::ThreadSafe>& HttpRequest : ActiveRequests)
+	{
+		if (HttpRequest->GetStatus() == EHttpRequestStatus::Processing)
+		{
+			HttpRequest->CancelRequest();
+		}
+	}
+
+	CalculateAverageBandwidth();
+	// Broadcast the measurement to any subscribers
+	OnBandwidthMeasurementComplete.Broadcast(CurrentBandwidthMeasurementInBps);
+
+	// Reset Values for next cycle
+	ActiveRequests.Empty();
+	MeasurementStartTimestamp = 0.0;
+	bHasTheLastRequestFinishedEarly = false;
+	bReadyToMeasureAfterWarmupPeriod = false;
+	TotalBytesReceived = 0;
+}
+
+void FBandwidthMeasurementTool::CalculateAverageBandwidth()
+{
+	// Calculate raw throughput for this measurement window
+	double ActualMeasurementWindowInSeconds = (MeasurementStartTimestamp > 0.0)
+		? (FPlatformTime::Seconds() - MeasurementStartTimestamp)
+		: 0.0;
+	ActualMeasurementWindowInSeconds = FMath::Max(ActualMeasurementWindowInSeconds, UE_DOUBLE_SMALL_NUMBER);
+	uint64 RawThroughputBps = static_cast<uint64>(static_cast<double>(TotalBytesReceived.load()) / ActualMeasurementWindowInSeconds);
+
+	TelemetryData.RecentReportedBytesReceived = TotalBytesReceived.load();
+	TelemetryData.RecentMeasurementTime = ActualMeasurementWindowInSeconds;
+	TelemetryData.RecentMeasurementInBps = RawThroughputBps;
+
+	AugmentBandwidthMeasurement(RawThroughputBps);
+	RecordBandwidthHistory(RawThroughputBps);
+
+	// Get history for averaging
+	const uint64 Count = static_cast<uint64>(RequestMemoryThroughputHistory.Num());
+	uint64 AverageBps = 0;
+	for (uint64 Measurement : RequestMemoryThroughputHistory)
+	{
+		// Divide each entry before accumulating to prevent rollover on the sum
+		AverageBps = SafeIncrementForUInt64(AverageBps, Measurement / Count);
+	}
+
+	TelemetryData.RecentMeasurementCount = Count;
+	TelemetryData.RecentAverageInBps = AverageBps;
+	TelemetryData.RecentHistory = RequestMemoryThroughputHistory;
+
+	CurrentBandwidthMeasurementInBps = AverageBps;
+}
+
+void FBandwidthMeasurementTool::AugmentBandwidthMeasurement(uint64& MeasurementToAlterInBytes) const
+{
+	// This will do nothing at the moment, but as we have need for certain cases,
+	// like platform considerations, they will go here
+}
+
+void FBandwidthMeasurementTool::RecordBandwidthHistory(const uint64 CalculatedBytesPerSecond)
+{
+	if (CalculatedBytesPerSecond == 0)
+	{
+		return;
+	}
+
+	if (RequestMemoryThroughputHistory.Num() >= BWMT::GMaxThroughputHistory)
+	{
+		// If we are at the limit we need to remove the oldest element
+		RequestMemoryThroughputHistory.RemoveAt(0);
+	}
+
+	RequestMemoryThroughputHistory.Add(CalculatedBytesPerSecond);
+}
+
+void FBandwidthMeasurementTool::ClearHistory()
+{
+	RequestMemoryThroughputHistory.Empty(BWMT::GMaxThroughputHistory);
+}
\ No newline at end of file
diff --git a/Engine/Source/Runtime/Online/Experimental/BandwidthMeasurementTool/Private/BandwidthMeasurementToolModule.cpp b/Engine/Source/Runtime/Online/Experimental/BandwidthMeasurementTool/Private/BandwidthMeasurementToolModule.cpp
new file mode 100644
--- /dev/null
+++ b/Engine/Source/Runtime/Online/Experimental/BandwidthMeasurementTool/Private/BandwidthMeasurementToolModule.cpp
@@ -0,0 +1,23 @@
+// Copyright Epic Games, Inc. All Rights Reserved.
+
+#include "Modules/ModuleInterface.h"
+#include "Modules/ModuleManager.h"
+#include "BandwidthMeasurementTool.h"
+
+// BandwidthMeasurementTool module
+// Module to manage the initialization and tear down for the Bandwidth measuring tool
+
+class FBandwidthMeasurementToolModule : public IModuleInterface
+{
+	void StartupModule()
+	{
+		FBandwidthMeasurementTool::Get().Initialize();
+	}
+
+	void ShutdownModule()
+	{
+		FBandwidthMeasurementTool::Get().Shutdown();
+	}
+};
+
+IMPLEMENT_MODULE(FBandwidthMeasurementToolModule, BandwidthMeasurementTool)
\ No newline at end of file
diff --git a/Engine/Source/Runtime/Online/Experimental/BandwidthMeasurementTool/Public/BandwidthMeasurementTool.h b/Engine/Source/Runtime/Online/Experimental/BandwidthMeasurementTool/Public/BandwidthMeasurementTool.h
new file mode 100644
--- /dev/null
+++ b/Engine/Source/Runtime/Online/Experimental/BandwidthMeasurementTool/Public/BandwidthMeasurementTool.h
@@ -0,0 +1,133 @@
+// Copyright Epic Games, Inc. All Rights Reserved.
+#pragma once
+#include "Containers/CircularQueue.h"
+#include "Containers/Ticker.h"
+#include "Templates/UnrealTemplate.h"
+
+#define UE_API BANDWIDTHMEASUREMENTTOOL_API
+
+UE_API DECLARE_LOG_CATEGORY_EXTERN(LogBandwidthMeasurementTool, Log, All);
+
+class IHttpRequest;
+
+namespace UE::BandwidthMeasurement
+{
+	// Delegate receives a measurement in Bytes Per Second
+	DECLARE_MULTICAST_DELEGATE_OneParam(FOnMeasurementComplete, uint64);
+}
+
+class FBandwidthMeasurementTool
+{
+public:
+	UE_NONCOPYABLE(FBandwidthMeasurementTool);
+
+	FBandwidthMeasurementTool();
+	~FBandwidthMeasurementTool();
+
+	// --- Singleton access ---
+	UE_API static FBandwidthMeasurementTool& Get();
+
+	// --- Lifecycle (called by module) ---
+	void Initialize();
+	void Shutdown();
+
+	// --- API Functions ---
+
+	// Assign to call back function to report when measurement is complete
+	UE_API void RegisterForCompletedMeasurement(TFunction<void(uint64)>&& FunctionPtr);
+
+	// Trigger measurement
+	UE_API void BeginMeasurement();
+
+	// Report out measurement completion
+	UE_API bool IsBandwidthBeingMeasured() const;
+
+	// Get delay time needed for measurement
+	UE_API float GetTimeNeededForMeasurement() const;
+
+	// --- Exposed Utility Functions ---
+	void ClearHistory();
+
+private:
+	// --- Internal Functions needed ---
+	
+	// Request Traffic from CDN
+	void CreateRequestTraffic();
+
+	// Tick-able Stop to call after time for measurement has passed
+	void StopMeasuring();
+
+	// Queue up bandwidth history
+	void RecordBandwidthHistory(const uint64 CalculatedBytesPerSecond);
+
+	// Calculate base average throughput
+	void CalculateAverageBandwidth();
+
+	// Apply any alterations to average throughput on case basis, maybe? Think platform needs
+	void AugmentBandwidthMeasurement(uint64& MeasurementToAlterInBytes) const;
+
+	// Helper function to manage uint64 rollover
+	uint64 SafeIncrementForUInt64(const uint64 Base, const uint64 Increase) const;
+
+	// --- Internal Members needed ---
+
+	// Singleton instance pointer
+	static TUniquePtr<FBandwidthMeasurementTool> Instance;
+
+	// Queue of fixed size for history average
+	TArray<uint64> RequestMemoryThroughputHistory;
+
+	// External reporting callback
+	UE::BandwidthMeasurement::FOnMeasurementComplete OnBandwidthMeasurementComplete;
+
+	// Current average throughput
+	uint64 CurrentBandwidthMeasurementInBps = 0;
+
+	// The total number of bytes received via http requests this measurement
+	std::atomic<uint64> TotalBytesReceived = 0;
+
+	// Fixed array of TSharedRef<IHttpRequest> to track request states
+	TArray<TSharedRef<IHttpRequest>> ActiveRequests;
+
+	// Tick handle needed to track current measurements time passage
+	FTSTicker::FDelegateHandle MeasurementTimer;
+
+	// Used to tell if the final request finished before hitting the maximum run time
+	bool bHasTheLastRequestFinishedEarly = false;
+
+	// Tick handle to track time for the "warm up" period
+	FTSTicker::FDelegateHandle WarmupTimer;
+
+	// State Tracking used to allow for communication line warmup
+	bool bReadyToMeasureAfterWarmupPeriod = false;
+
+	// Wall-clock timestamp (FPlatformTime::Seconds) when post-warmup measurement began
+	double MeasurementStartTimestamp = 0.0;
+
+	// Mutex to keep data consistent across multiple potential thread executions
+	mutable FTransactionallySafeCriticalSection Mutex;
+
+public:
+	// --- Telemetry Data ---
+	struct FBandwidthSnapshotData
+	{
+		uint64 FinishRequestCount = 0;
+		uint64 CompleteRequestCount = 0;
+		uint64 PostWarmupCompleteRequestCount = 0;
+		uint64 RecentReportedBytesReceived = 0;
+		double RecentMeasurementTime = 0;
+		uint64 RecentMeasurementInBps = 0;
+		uint64 RecentMeasurementCount = 0;
+		uint64 RecentAverageInBps = 0;
+		TArray<uint64> RecentHistory;
+	};
+
+	// Returns Telemetry Data for reporting
+	const FBandwidthSnapshotData& GetTelemetryData() const { return TelemetryData; };
+
+private:
+	// Storage structure for telemetry reporting
+	FBandwidthSnapshotData TelemetryData;
+};
+
+#undef UE_API
\ No newline at end of file
diff --git a/Engine/Source/Runtime/Online/Experimental/DownlinkBandwidthManager/DownlinkBandwidthManager.Build.cs b/Engine/Source/Runtime/Online/Experimental/DownlinkBandwidthManager/DownlinkBandwidthManager.Build.cs
--- a/Engine/Source/Runtime/Online/Experimental/DownlinkBandwidthManager/DownlinkBandwidthManager.Build.cs
+++ b/Engine/Source/Runtime/Online/Experimental/DownlinkBandwidthManager/DownlinkBandwidthManager.Build.cs
@@ -15,6 +15,7 @@ public class DownlinkBandwidthManager : ModuleRules
 
 		PrivateDependencyModuleNames.AddRange(new string[] {
 			"CoreUObject",
+			"BandwidthMeasurementTool",
 		});
 	}
 }
diff --git a/Engine/Source/Runtime/Online/Experimental/DownlinkBandwidthManager/Private/DownlinkBandwidthManager.cpp b/Engine/Source/Runtime/Online/Experimental/DownlinkBandwidthManager/Private/DownlinkBandwidthManager.cpp
--- a/Engine/Source/Runtime/Online/Experimental/DownlinkBandwidthManager/Private/DownlinkBandwidthManager.cpp
+++ b/Engine/Source/Runtime/Online/Experimental/DownlinkBandwidthManager/Private/DownlinkBandwidthManager.cpp
@@ -3,7 +3,7 @@
 #include "DownlinkBandwidthManager.h"
 #include "HAL/IConsoleManager.h"
 #include "Math/NumericLimits.h"
-//#include "Misc/ConfigCacheIni.h"
+#include "BandwidthMeasurementTool.h"
 
 DEFINE_LOG_CATEGORY(LogDownlinkBandwidthManager);
 
@@ -14,14 +14,14 @@ namespace DLBW
 	static FAutoConsoleVariableRef CVarDLBWDefaultBandwidthToAllocate(
 		TEXT("net.DLBW.DefaultBandwidth"),
 		GDefaultBandwidthToAllocate,
-		TEXT("Default value to use for bandwidth cap distribution in Kbps. Example: Average american home internet is 200-300Mbps, so we would use 200,000 - 300,000 for this value."),
+		TEXT("Default value to use for bandwidth cap distribution in Kbps. Example: Average American home internet is 200-300Mbps, so we would use 200,000 - 300,000 for this value."),
 		ECVF_Default);
 
 	int32 GMinRequiredBandwidthDuringGameplay = 50000; // 50 Mbps
 	static FAutoConsoleVariableRef CVarDLBWMinRequiredBandwidthDuringGameplay(
 		TEXT("net.DLBW.MinGameplayBandwidth"),
 		GMinRequiredBandwidthDuringGameplay,
-		TEXT("Default value to use for bandwidth cap distribution in Kbps. Example: Average american home internet is 200-300Mbps, so we would use 200,000 - 300,000 for this value."),
+		TEXT("Default value to use for bandwidth cap distribution in Kbps. Example: Average American home internet is 200-300Mbps, so we would use 200,000 - 300,000 for this value."),
 		ECVF_Default);
 
 	static TAutoConsoleVariable<int32> CVarDLBWPriorityWeightExponent(
@@ -33,15 +33,26 @@ namespace DLBW
 	static TAutoConsoleVariable<float> CVarDLBWDistributionPolling(
 		TEXT("net.DLBW.DistributionPolling"),
 		1.0f,
-		TEXT("Dealy Time, per frame cycle to poll for changes needed in distribution."),
+		TEXT("Delay Time, per frame cycle to poll for changes needed in distribution."),
 		ECVF_Default);
 
-	static bool GEnableBandwidthLogging = false;
-	static FAutoConsoleVariableRef CVarDLBWEnableBandwidthLogging(
-		TEXT("net.DLBW.EnableLogging"),
-		GEnableBandwidthLogging,
-		TEXT("Turns on Logging for the manager, and will write out based on CVarDLBWDistributionPolling rate."),
+	float GBandwidthPercentageToUse = 0.8f; 
+	static FAutoConsoleVariableRef CVarDLBWBandwidthPercentToUse(
+		TEXT("net.DLBW.BandwidthPercentage"),
+		GBandwidthPercentageToUse,
+		TEXT("0 - 1 value used to control the amount of measured bandwidth the manager will utilize."),
 		ECVF_Default);
+
+	// Command to print logging information
+	FAutoConsoleCommand FDLBWShowBandwidthTelemetry(
+		TEXT("net.DLBW.PrintTelemetry"),
+		TEXT("Prints a snapshot of the manager data for logging purposes."),
+		FConsoleCommandDelegate::CreateLambda([]()
+			{
+				FDownlinkBandwidthManager& DLBWManager = FDownlinkBandwidthManager::Get();
+				DLBWManager.Log();
+			}),
+		ECVF_Cheat);
 }
 
 /// Non-Member Functions
@@ -52,7 +63,7 @@ FString LexToString(UE::BandwidthManager::EBandwidthPriority Value)
 	case UE::BandwidthManager::EBandwidthPriority::HighPriority:
 		return TEXT("High Priority");
 	case UE::BandwidthManager::EBandwidthPriority::MediumPriority:
-		return TEXT("Medium Prority");
+		return TEXT("Medium Priority");
 	case UE::BandwidthManager::EBandwidthPriority::LowPriority:
 		return TEXT("Low Priority");
 	default:
@@ -89,12 +100,31 @@ FDownlinkBandwidthManager::~FDownlinkBandwidthManager()
 
 void FDownlinkBandwidthManager::Initialize()
 {
+	// Set an initial value
 	SetIncomingBandwidthCap(DLBW::GDefaultBandwidthToAllocate);
 	
+	FBandwidthMeasurementTool& BandwidthTool = FBandwidthMeasurementTool::Get();
+	BandwidthTool.RegisterForCompletedMeasurement([this](uint64 MeasurementInBps)
+		{
+			const int32 RestrictedCap = static_cast<int32>(static_cast<float>(MeasurementInBps) *
+				FMath::Clamp(DLBW::GBandwidthPercentageToUse, 0.0f, 1.0f));
+
+			const uint64 MeasurementInKbps = RestrictedCap / UE::BandwidthManager::KilobitsToBytesConversionValue;
+			SetIncomingBandwidthCap(static_cast<int32>(MeasurementInKbps));
+		});
+
+	DLBW::CVarDLBWDefaultBandwidthToAllocate->SetOnChangedCallback(FConsoleVariableDelegate::CreateLambda(
+		[this](IConsoleVariable* ConsoleVariable) 
+		{
+			int32 NewCap;
+			ConsoleVariable->GetValue(NewCap);
+			SetIncomingBandwidthCap(NewCap);
+		}));
+
 	// return true for reschedule
 	TickHandle = FTSTicker::GetCoreTicker().AddTicker(FTickerDelegate::CreateLambda([this](float InDeltaSeconds)->bool
 		{
-			// let's us evaulate distribution on time cycle needs should anything change
+			// let's us evaluate distribution on time cycle needs should anything change
 			Tick();
 			return true;
 		}), DLBW::CVarDLBWDistributionPolling.GetValueOnAnyThread());
@@ -139,7 +169,7 @@ UE::BandwidthManager::FServiceHandleID FDownlinkBandwidthManager::RegisterMonito
 		{
 			uint8 ValueA = (CurrentGameplayPriorityState == UE::BandwidthManager::EBandwidthDistributionPriorityState::IdleGameplay) ? ServiceA.ConfigData.IdleGameplayPriority : ServiceA.ConfigData.InGameplayPriority;
 			uint8 ValueB = (CurrentGameplayPriorityState == UE::BandwidthManager::EBandwidthDistributionPriorityState::IdleGameplay) ? ServiceB.ConfigData.IdleGameplayPriority : ServiceB.ConfigData.InGameplayPriority;
-			// Sorts in ascending priority, so lower numerical (higher listed) services can be alloted to first
+			// Sorts in ascending priority, so lower numerical (higher listed) services can be allotted to first
 			return ValueA > ValueB;
 		});
 
@@ -253,21 +283,8 @@ void FDownlinkBandwidthManager::SetDistributionGameplayState(UE::BandwidthManage
 
 void FDownlinkBandwidthManager::Tick()
 {
-	if (CapKbps != DLBW::GDefaultBandwidthToAllocate)
-	{
-		// Updates if Cvar is changed, and will handle DistributeBandwidth
-		SetIncomingBandwidthCap(DLBW::GDefaultBandwidthToAllocate);
-	}
-	else
-	{
-		// Called as a background poll to see if redistribution is needed
-		DistributeBandwidth();
-	}
-
-	if (DLBW::GEnableBandwidthLogging)
-	{
-		Log();
-	}
+	// Called as a background poll to see if redistribution is needed
+	DistributeBandwidth();
 }
 
 void FDownlinkBandwidthManager::GeneratePriorityBucketMap()
@@ -277,8 +294,8 @@ void FDownlinkBandwidthManager::GeneratePriorityBucketMap()
 	int32 BandwidthBudgetBytes = GetCurrentMaxAvailableBandwidthInBytes();
 	if (BandwidthBudgetBytes <= 0)
 	{
-		UE_LOG(LogDownlinkBandwidthManager, Warning, TEXT("%s - BandwidthBudgetPerFrame was set to less than 0. Value: %i, CapKbps: %i"),
-			ANSI_TO_TCHAR(__FUNCTION__), BandwidthBudgetBytes, CapKbps);
+		UE_LOGF(LogDownlinkBandwidthManager, Warning, "%ls - BandwidthBudgetPerFrame was set to less than 0. Value: %i, CapKbps: %i",
+			UE_SOURCE_LOCATION, BandwidthBudgetBytes, CapKbps);
 		return;
 	}
 
@@ -323,14 +340,14 @@ void FDownlinkBandwidthManager::GeneratePriorityBucketMap()
 		// PriorityList = [High, Medium, Low], WeightBiasExponent = 3
 		// This would yield a starting WeightedPercentage = 0.702
 		// The next Cycle would be interpolate from 0 - 0.702 and then WeightedPercentage = 0.209
-		// The Overal Reduction becomes 0.911, meaning the final bucket gets 1-0.911 or 0.089
+		// The Overall Reduction becomes 0.911, meaning the final bucket gets 1-0.911 or 0.089
 		// In Percentages of the total bandwidth that is 70.2, 20.9, and 8.9%
 		// This is the total bandwidth each bucket is allowed and that bucket then needs to
 		// subdivide its allowance across it's own services
 
 		if (++BucketNumber == BucketCount)
 		{
-			// Last Item, we want to give out what's left to maximize bandwidth useage
+			// Last Item, we want to give out what's left to maximize bandwidth usage
 			BucketRef.BandwidthAllocationForPriority = BandwidthBudgetBytes;
 		}
 		else
@@ -358,8 +375,8 @@ void FDownlinkBandwidthManager::DistributeBandwidth()
 	int32 BandwidthBudgetBytes = GetCurrentMaxAvailableBandwidthInBytes();
 	if (BandwidthBudgetBytes <= 0)
 	{
-		UE_LOG(LogDownlinkBandwidthManager, Warning, TEXT("%s - BandwidthBudgetPerFrame was set to less than 0. Value: %i, CapKbps: %i"), 
-			ANSI_TO_TCHAR(__FUNCTION__), BandwidthBudgetBytes, CapKbps);
+		UE_LOGF(LogDownlinkBandwidthManager, Warning, "%ls - BandwidthBudgetPerFrame was set to less than 0. Value: %i, CapKbps: %i", 
+			UE_SOURCE_LOCATION, BandwidthBudgetBytes, CapKbps);
 		return;
 	}
 
@@ -420,45 +437,45 @@ UE::BandwidthManager::EBandwidthPriority FDownlinkBandwidthManager::GetPriorityV
 
 void FDownlinkBandwidthManager::Log()
 {
-	UE_LOG(LogDownlinkBandwidthManager, Log, TEXT("FDownlinkBandwidthManager Logging Breakdown:\n"));
+	UE_LOGF(LogDownlinkBandwidthManager, Log, "FDownlinkBandwidthManager Logging Breakdown:\n");
 
-	UE_LOG(LogDownlinkBandwidthManager, Log, TEXT("FDownlinkBandwidthManager Internal Data:"));
-	UE_LOG(LogDownlinkBandwidthManager, Log, TEXT("Default Bandwidth Cap in Kbps: %i"), static_cast<int>(DLBW::GDefaultBandwidthToAllocate));
-	UE_LOG(LogDownlinkBandwidthManager, Log, TEXT("Minimum Bandwidth Cap during gameplay in Kbps: %i"), static_cast<int>(DLBW::GMinRequiredBandwidthDuringGameplay));
-	UE_LOG(LogDownlinkBandwidthManager, Log, TEXT("Current Bandwidth Cap in Kbps: %i"), static_cast<int>(CapKbps));
-	UE_LOG(LogDownlinkBandwidthManager, Log, TEXT("Current Gameplay Priority State: %s"), 
+	UE_LOGF(LogDownlinkBandwidthManager, Log, "FDownlinkBandwidthManager Internal Data:");
+	UE_LOGF(LogDownlinkBandwidthManager, Log, "Default Bandwidth Cap in Kbps: %i", static_cast<int>(DLBW::GDefaultBandwidthToAllocate));
+	UE_LOGF(LogDownlinkBandwidthManager, Log, "Minimum Bandwidth Cap during gameplay in Kbps: %i", static_cast<int>(DLBW::GMinRequiredBandwidthDuringGameplay));
+	UE_LOGF(LogDownlinkBandwidthManager, Log, "Current Bandwidth Cap in Kbps: %i", static_cast<int>(CapKbps));
+	UE_LOGF(LogDownlinkBandwidthManager, Log, "Current Gameplay Priority State: %ls", 
 		(CurrentGameplayPriorityState == UE::BandwidthManager::EBandwidthDistributionPriorityState::IdleGameplay) ? TEXT("Outside Normal Gameplay") : TEXT("In Active Gameplay"));
-	UE_LOG(LogDownlinkBandwidthManager, Log, TEXT("Tick Polling in seconds: %0.3f"), DLBW::CVarDLBWDistributionPolling.GetValueOnAnyThread());
+	UE_LOGF(LogDownlinkBandwidthManager, Log, "Tick Polling in seconds: %0.3f", DLBW::CVarDLBWDistributionPolling.GetValueOnAnyThread());
 
-	UE_LOG(LogDownlinkBandwidthManager, Log, TEXT("FDownlinkBandwidthManager Prioritization:"));
-	UE_LOG(LogDownlinkBandwidthManager, Log, TEXT("Priority Weight Exponent: %i"), static_cast<int>(DLBW::CVarDLBWPriorityWeightExponent.GetValueOnAnyThread()));
+	UE_LOGF(LogDownlinkBandwidthManager, Log, "FDownlinkBandwidthManager Prioritization:");
+	UE_LOGF(LogDownlinkBandwidthManager, Log, "Priority Weight Exponent: %i", static_cast<int>(DLBW::CVarDLBWPriorityWeightExponent.GetValueOnAnyThread()));
 
 	for (TPair<UE::BandwidthManager::EBandwidthPriority, UE::BandwidthManager::FPriorityBucket>& BucketData : PriorityBucketMapping)
 	{
-		UE_LOG(LogDownlinkBandwidthManager, Log, TEXT("Bucket Priority Value: %s"), *LexToString(BucketData.Key));
-		UE_LOG(LogDownlinkBandwidthManager, Log, TEXT("Bucket Service Count: %u"), static_cast<unsigned int>(BucketData.Value.NumberOfServicesAtPriority));
-		UE_LOG(LogDownlinkBandwidthManager, Log, TEXT("Bandwidth for this Bucket: %i"), static_cast<int>(BucketData.Value.BandwidthAllocationForPriority));
+		UE_LOGF(LogDownlinkBandwidthManager, Log, "Bucket Priority Value: %ls", *LexToString(BucketData.Key));
+		UE_LOGF(LogDownlinkBandwidthManager, Log, "Bucket Service Count: %u", static_cast<unsigned int>(BucketData.Value.NumberOfServicesAtPriority));
+		UE_LOGF(LogDownlinkBandwidthManager, Log, "Bandwidth for this Bucket: %i", static_cast<int>(BucketData.Value.BandwidthAllocationForPriority));
 	}
 
-	UE_LOG(LogDownlinkBandwidthManager, Log, TEXT("FDownlinkBandwidthManager Tracked Services:"));
+	UE_LOGF(LogDownlinkBandwidthManager, Log, "FDownlinkBandwidthManager Tracked Services:");
 	for (TPair<UE::BandwidthManager::FServiceHandleID, FManagedService>& ServiceData : TrackedServices)
 	{
-		UE_LOG(LogDownlinkBandwidthManager, Log, TEXT("Tracked Service ID: %i, Name: %s"), static_cast<int>(ServiceData.Key), *ServiceData.Value.DebugName);
-		UE_LOG(LogDownlinkBandwidthManager, Log, TEXT("Service Gameplay Priority Value: %s"), *LexToString(ServiceData.Value.ConfigData.InGameplayPriority));
-		UE_LOG(LogDownlinkBandwidthManager, Log, TEXT("Service Gameplay Priority Value: %s"), *LexToString(ServiceData.Value.ConfigData.IdleGameplayPriority));
-		UE_LOG(LogDownlinkBandwidthManager, Log, TEXT("Is Service Currently Active: %s"), (ServiceData.Value.bActiveConnection) ? TEXT("True") : TEXT("False"));
+		UE_LOGF(LogDownlinkBandwidthManager, Log, "Tracked Service ID: %i, Name: %ls", static_cast<int>(ServiceData.Key), *ServiceData.Value.DebugName);
+		UE_LOGF(LogDownlinkBandwidthManager, Log, "Service Gameplay Priority Value: %ls", *LexToString(ServiceData.Value.ConfigData.InGameplayPriority));
+		UE_LOGF(LogDownlinkBandwidthManager, Log, "Service Gameplay Priority Value: %ls", *LexToString(ServiceData.Value.ConfigData.IdleGameplayPriority));
+		UE_LOGF(LogDownlinkBandwidthManager, Log, "Is Service Currently Active: %ls", (ServiceData.Value.bActiveConnection) ? TEXT("True") : TEXT("False"));
 		if (!ServiceData.Value.ConfigData.ComparisonValue.IsValid())
 		{
-			UE_LOG(LogDownlinkBandwidthManager, Log, TEXT("Service Comparision Value: Null"));
+			UE_LOGF(LogDownlinkBandwidthManager, Log, "Service Comparison Value: Null");
 		}
 		else
 		{
-			UE_LOG(LogDownlinkBandwidthManager, Log, TEXT("Service Comparision Value: %i"), static_cast<int>(*ServiceData.Value.ConfigData.ComparisonValue.Get()));
+			UE_LOGF(LogDownlinkBandwidthManager, Log, "Service Comparison Value: %i", static_cast<int>(*ServiceData.Value.ConfigData.ComparisonValue.Get()));
 		}
 
-		UE_LOG(LogDownlinkBandwidthManager, Log, TEXT("Service Alloted Bandwidth in Kbps: %i"), static_cast<int>(ServiceData.Value.AllotedBandwidthKbps));
-		UE_LOG(LogDownlinkBandwidthManager, Log, TEXT("Service Bound Callback: %s\n"), (ServiceData.Value.OnBandwidthUpdate.IsSet()) ? TEXT("Set") : TEXT("Unset"));
+		UE_LOGF(LogDownlinkBandwidthManager, Log, "Service Allotted Bandwidth in Kbps: %i", static_cast<int>(ServiceData.Value.AllotedBandwidthKbps));
+		UE_LOGF(LogDownlinkBandwidthManager, Log, "Service Bound Callback: %ls\n", (ServiceData.Value.OnBandwidthUpdate.IsSet()) ? TEXT("Set") : TEXT("Unset"));
 	}
 
-	UE_LOG(LogDownlinkBandwidthManager, Log, TEXT("FDownlinkBandwidthManager Logging Complete."));
+	UE_LOGF(LogDownlinkBandwidthManager, Log, "FDownlinkBandwidthManager Logging Complete.");
 }
\ No newline at end of file
