[Http]Move EasyHandle creation to SetupRequest and destruction to CleanupRequest in FCurlHttpRequest.

#Synced-CL 52039227

#Synced-CL 52039227

diff --git a/Engine/Source/Runtime/Online/HTTP/Private/Curl/CurlHttp.cpp b/Engine/Source/Runtime/Online/HTTP/Private/Curl/CurlHttp.cpp
--- a/Engine/Source/Runtime/Online/HTTP/Private/Curl/CurlHttp.cpp
+++ b/Engine/Source/Runtime/Online/HTTP/Private/Curl/CurlHttp.cpp
@@ -43,6 +43,13 @@ TAutoConsoleVariable<bool> CVarHttpCurlResetRequestCompletedOnRetryEnabled(
 	ECVF_SaveForNextBoot
 );
 
+TAutoConsoleVariable<bool> CVarHttpCurlDeferEasyHandleLifetime(
+	TEXT("http.CurlDeferEasyHandleLifetime"),
+	true,
+	TEXT("When enabled, the curl easy handle is created in SetupRequest and destroyed in CleanupRequest rather than in the constructor/destructor. This prevents the handle from outliving the curl library during shutdown."),
+	ECVF_SaveForNextBoot
+);
+
 #if WITH_SSL
 static int SslCertVerify(int PreverifyOk, X509_STORE_CTX* Context)
 {
@@ -102,9 +109,42 @@ FCurlHttpRequest::FCurlHttpRequest()
 	, LastReportedBytesSent(0)
 	, LeastRecentlyCachedInfoMessageIndex(0)
 {
-	checkf(FCurlHttpManager::IsInit(), TEXT("Curl request was created while the library is shutdown"));
+	InfoMessageCache.AddDefaulted(NumberOfInfoMessagesToCache);
+
+	if (!CVarHttpCurlDeferEasyHandleLifetime.GetValueOnAnyThread())
+	{
+		checkf(FCurlHttpManager::IsInit(), TEXT("Curl request was created while the library is shutdown"));
+		InitEasyHandle();
+	}
+
+	// Add default headers
+	const TMap<FString, FString>& DefaultHeaders = FHttpModule::Get().GetDefaultHeaders();
+	for (TMap<FString, FString>::TConstIterator It(DefaultHeaders); It; ++It)
+	{
+		SetHeader(It.Key(), It.Value());
+	}
+
+	bUsePlatformActivityTimeout = false;
+}
+
+FCurlHttpRequest::~FCurlHttpRequest()
+{
+	if (!CVarHttpCurlDeferEasyHandleLifetime.GetValueOnAnyThread())
+	{
+		DestroyEasyHandle();
+	}
+}
+
+bool FCurlHttpRequest::InitEasyHandle()
+{
+	check(!EasyHandle);
 
 	EasyHandle = curl_easy_init();
+	if (!EasyHandle)
+	{
+		UE_LOG(LogHttp, Error, TEXT("%p: Failed to initialize curl easy handle"), this);
+		return false;
+	}
 
 	// Always setup the debug function to allow for activity to be tracked
 	curl_easy_setopt(EasyHandle, CURLOPT_DEBUGDATA, this);
@@ -201,21 +241,11 @@ FCurlHttpRequest::FCurlHttpRequest()
 	curl_easy_setopt(EasyHandle, CURLOPT_SSL_CTX_DATA, this);
 #endif // #if WITH_SSL
 
-	InfoMessageCache.AddDefaulted(NumberOfInfoMessagesToCache);
-
-	// Add default headers
-	const TMap<FString, FString>& DefaultHeaders = FHttpModule::Get().GetDefaultHeaders();
-	for (TMap<FString, FString>::TConstIterator It(DefaultHeaders); It; ++It)
-	{
-		SetHeader(It.Key(), It.Value());
-	}
-
-	bUsePlatformActivityTimeout = false;
+	return true;
 }
 
-FCurlHttpRequest::~FCurlHttpRequest()
+void FCurlHttpRequest::DestroyEasyHandle()
 {
-	checkf(FCurlHttpManager::IsInit(), TEXT("Curl request was held after the library was shutdown."));
 	if (EasyHandle)
 	{
 		// cleanup the handle first (that order is used in howtos)
@@ -278,7 +308,6 @@ void FCurlHttpRequest::SetVerb(const FString& InVerb)
 		return;
 	}
 
-	check(EasyHandle);
 	Verb = InVerb.ToUpper();
 }
 
@@ -290,7 +319,6 @@ void FCurlHttpRequest::SetOption(const FName Option, const FString& OptionValue)
 		return;
 	}
 
-	check(EasyHandle);
 	FHttpRequestCommon::SetOption(Option, OptionValue);
 }
 
@@ -823,6 +851,17 @@ void FCurlHttpRequest::OnAnyActivityOccur(FStringView Reason)
 bool FCurlHttpRequest::SetupRequest()
 {
 	QUICK_SCOPE_CYCLE_COUNTER(STAT_FCurlHttpRequest_SetupRequest);
+
+	if (CVarHttpCurlDeferEasyHandleLifetime.GetValueOnAnyThread())
+	{
+		checkf(FCurlHttpManager::IsInit(), TEXT("Curl request was created while the library is shutdown"));
+
+		if (!InitEasyHandle())
+		{
+			return false;
+		}
+	}
+
 	check(EasyHandle);
 
 	if (!RequestPayload.IsValid())
@@ -1074,7 +1113,6 @@ void FCurlHttpRequest::CleanupRequestHttpThread()
 bool FCurlHttpRequest::ProcessRequest()
 {
 	QUICK_SCOPE_CYCLE_COUNTER(STAT_FCurlHttpRequest_ProcessRequest);
-	check(EasyHandle);
 
 	// Clear out response. If this is a re-used request, Response could point to a stale response until SetupRequestHttpThread is called
 	ResponseCommon = nullptr;
@@ -1318,6 +1356,11 @@ void FCurlHttpRequest::LogCurlErrors()
 void FCurlHttpRequest::CleanupRequest()
 {
 	CloseRequestPayloadDefaultImpl();
+
+	if (CVarHttpCurlDeferEasyHandleLifetime.GetValueOnAnyThread())
+	{
+		DestroyEasyHandle();
+	}
 }
 
 FHttpResponsePtr FCurlHttpRequest::CreateResponse()
diff --git a/Engine/Source/Runtime/Online/HTTP/Private/Curl/CurlHttp.h b/Engine/Source/Runtime/Online/HTTP/Private/Curl/CurlHttp.h
--- a/Engine/Source/Runtime/Online/HTTP/Private/Curl/CurlHttp.h
+++ b/Engine/Source/Runtime/Online/HTTP/Private/Curl/CurlHttp.h
@@ -126,6 +126,20 @@ protected:
 
 private:
 
+	/**
+	 * Initialize the curl easy handle and apply all static options.
+	 * Called either from the constructor (legacy path) or SetupRequest (deferred path).
+	 *
+	 * @return true if the handle was successfully created
+	 */
+	bool InitEasyHandle();
+
+	/**
+	 * Destroy the curl easy handle and free the header list.
+	 * Called either from the destructor (legacy path) or CleanupRequest (deferred path).
+	 */
+	void DestroyEasyHandle();
+
 	/**
 	 * Static callback to be used as read function (CURLOPT_READFUNCTION), will dispatch the call to proper instance
 	 *
