[Mobile] Add throughput-based watchdog to cancel and re-queue slow background downloads

Also added a rotation function to move the problematic CDN to the end of the list so the retry is done in another CDN.

#rnx
#Synced-CL 52163679

diff --git a/Engine/Source/Runtime/Online/BackgroundHTTP/Private/BackgroundHttpRequestImpl.cpp b/Engine/Source/Runtime/Online/BackgroundHTTP/Private/BackgroundHttpRequestImpl.cpp
--- a/Engine/Source/Runtime/Online/BackgroundHTTP/Private/BackgroundHttpRequestImpl.cpp
+++ b/Engine/Source/Runtime/Online/BackgroundHTTP/Private/BackgroundHttpRequestImpl.cpp
@@ -146,6 +146,19 @@ const TArray<FString>& FBackgroundHttpRequestImpl::GetURLList() const
 	return URLList;
 }
 
+void FBackgroundHttpRequestImpl::RotateURLList()
+{
+	// Only safe to call when the request is not being processed (i.e., after CancelRequest and before ProcessRequest).
+	// Callers must ensure this precondition.
+	if (URLList.Num() > 1)
+	{
+		UE_LOG(LogBackgroundHttpRequest, Verbose, TEXT("Rotating URL list for request %s (moving %s to end)"), *GetRequestID(), *URLList[0]);
+		FString First = MoveTemp(URLList[0]);
+		URLList.RemoveAt(0, EAllowShrinking::No);
+		URLList.Add(MoveTemp(First));
+	}
+}
+
 void FBackgroundHttpRequestImpl::SetExpectedResultSize(const uint64 ExpectedSize)
 {
 	ExpectedResultSize = ExpectedSize;
diff --git a/Engine/Source/Runtime/Online/BackgroundHTTP/Public/BackgroundHttpRequestImpl.h b/Engine/Source/Runtime/Online/BackgroundHTTP/Public/BackgroundHttpRequestImpl.h
--- a/Engine/Source/Runtime/Online/BackgroundHTTP/Public/BackgroundHttpRequestImpl.h
+++ b/Engine/Source/Runtime/Online/BackgroundHTTP/Public/BackgroundHttpRequestImpl.h
@@ -32,6 +32,7 @@ public:
     BACKGROUNDHTTP_API virtual void ResumeRequest() override;
     BACKGROUNDHTTP_API virtual void SetURLAsList(const TArray<FString>& URLs, int NumRetriesToAttempt) override;
 	BACKGROUNDHTTP_API virtual const TArray<FString>& GetURLList() const override;
+	BACKGROUNDHTTP_API virtual void RotateURLList() override;
 	BACKGROUNDHTTP_API virtual void SetExpectedResultSize(const uint64 ExpectedSize) override;
 	BACKGROUNDHTTP_API virtual uint64 GetExpectedResultSize() const override;
 	BACKGROUNDHTTP_API virtual void SetCompleteNotification(FBackgroundHttpNotificationObjectPtr DownloadCompleteNotificationObjectIn) override;
diff --git a/Engine/Source/Runtime/Online/BackgroundHTTP/Public/Interfaces/IBackgroundHttpRequest.h b/Engine/Source/Runtime/Online/BackgroundHTTP/Public/Interfaces/IBackgroundHttpRequest.h
--- a/Engine/Source/Runtime/Online/BackgroundHTTP/Public/Interfaces/IBackgroundHttpRequest.h
+++ b/Engine/Source/Runtime/Online/BackgroundHTTP/Public/Interfaces/IBackgroundHttpRequest.h
@@ -125,6 +125,13 @@ public:
 	*/
 	virtual const TArray<FString>& GetURLList() const = 0;
 
+	/**
+	* Rotates the URL list by moving the first URL to the end.
+	* Useful for retrying a request against a different CDN after a slow or failed attempt.
+	* Should only be called when the request is not being processed.
+	*/
+	virtual void RotateURLList() {}
+
 	/**
 	* Sets expected result size for the request, might be used by the implementation
 	* to check file size before invoking completion notification. 
