Fix the warning caused by race condition in EventLoopNotifierEventFd during shutdown.

#Synced-CL 52039227

diff --git a/Engine/Source/Runtime/Online/EventLoop/Private/EventLoop/BSDSocket/EventLoopNotifierEventFd.cpp b/Engine/Source/Runtime/Online/EventLoop/Private/EventLoop/BSDSocket/EventLoopNotifierEventFd.cpp
--- a/Engine/Source/Runtime/Online/EventLoop/Private/EventLoop/BSDSocket/EventLoopNotifierEventFd.cpp
+++ b/Engine/Source/Runtime/Online/EventLoop/Private/EventLoop/BSDSocket/EventLoopNotifierEventFd.cpp
@@ -21,8 +21,6 @@ static FAutoConsoleVariableRef CVarHttpEventLoopNotifierUseThreadPool(
 	ECVF_SaveForNextBoot
 );
 
-std::atomic<bool> FEventLoopNotifierEventFd::bNotified = false;
-
 bool FEventLoopNotifierEventFd::Init()
 {
 	FileDescriptorRead = eventfd(0, EFD_NONBLOCK);
@@ -31,6 +29,15 @@ bool FEventLoopNotifierEventFd::Init()
 
 void FEventLoopNotifierEventFd::Shutdown()
 {
+	// Wait for any in-flight async notify to complete before closing the fd,
+	// so the lambda cannot access a closed or destroyed fd.
+#if PLATFORM_WRITES_ARE_SLOW
+	while (bNotified)
+	{
+		FPlatformProcess::Yield();
+	}
+#endif
+
 	if (FileDescriptorRead != INVALID_SOCKET)
 	{
 		close(FileDescriptorRead);
@@ -63,9 +70,9 @@ void FEventLoopNotifierEventFd::Notify()
 		// When Notify() get called several times in same frame, we only need to notify once
 		if (!bNotified.exchange(true))
 		{
-			AsyncPool(*GIOThreadPool, [Fd = FileDescriptorRead] {
-				bNotified = false;
+			AsyncPool(*GIOThreadPool, [Fd = FileDescriptorRead, &bNotified = this->bNotified] {
 				EventLoopNotifierEventFd::NotifyImpl(Fd);
+				bNotified = false;
 			});
 		}
 	}
diff --git a/Engine/Source/Runtime/Online/EventLoop/Public/EventLoop/BSDSocket/EventLoopNotifierEventFd.h b/Engine/Source/Runtime/Online/EventLoop/Public/EventLoop/BSDSocket/EventLoopNotifierEventFd.h
--- a/Engine/Source/Runtime/Online/EventLoop/Public/EventLoop/BSDSocket/EventLoopNotifierEventFd.h
+++ b/Engine/Source/Runtime/Online/EventLoop/Public/EventLoop/BSDSocket/EventLoopNotifierEventFd.h
@@ -23,7 +23,7 @@ public:
 
 private:
 	int32 FileDescriptorRead = INVALID_SOCKET;
-	static std::atomic<bool> bNotified;
+	std::atomic<bool> bNotified = false;
 };
 
 /* UE::EventLoop */ }
