[Http]Add FHttpModule bWasEverLoaded guard to prevent silent re-creation after shutdown.

#Synced-CL 52039227

diff --git a/Engine/Source/Runtime/Online/HTTP/Private/HttpModule.cpp b/Engine/Source/Runtime/Online/HTTP/Private/HttpModule.cpp
--- a/Engine/Source/Runtime/Online/HTTP/Private/HttpModule.cpp
+++ b/Engine/Source/Runtime/Online/HTTP/Private/HttpModule.cpp
@@ -41,6 +41,7 @@ DEFINE_LOG_CATEGORY(LogHttp);
 IMPLEMENT_MODULE(FHttpModule, HTTP);
 
 FHttpModule* FHttpModule::Singleton = nullptr;
+bool FHttpModule::bWasEverLoaded = false;
 
 static bool ShouldLaunchUrl(const TCHAR* Url)
 {
@@ -89,6 +90,7 @@ void FHttpModule::UpdateConfigs()
 void FHttpModule::StartupModule()
 {
 	Singleton = this;
+	bWasEverLoaded = true;
 
 	MaxReadBufferSize = 256 * 1024;
 	HttpTotalTimeout = 0.0f;
@@ -369,6 +371,9 @@ FHttpModule& FHttpModule::Get()
 {
 	if (Singleton == nullptr)
 	{
+		// Only auto-load on first use. If the module was previously loaded and shut down
+		// (e.g. by a test fixture), do not recreate it — callers must explicitly load it again.
+		checkf(!bWasEverLoaded, TEXT("FHttpModule::Get() called after the HTTP module was shut down. Recreating it via Get() is not allowed. Call FModuleManager::LoadModuleChecked<FHttpModule>(\"HTTP\") explicitly if recreation is intended."));
 		check(IsInGameThread());
 		FModuleManager::LoadModuleChecked<FHttpModule>("HTTP");
 	}
diff --git a/Engine/Source/Runtime/Online/HTTP/Public/HttpModule.h b/Engine/Source/Runtime/Online/HTTP/Public/HttpModule.h
--- a/Engine/Source/Runtime/Online/HTTP/Public/HttpModule.h
+++ b/Engine/Source/Runtime/Online/HTTP/Public/HttpModule.h
@@ -359,6 +359,8 @@ private:
 	TMap<FString, FString> DefaultHeaders;
 	/** singleton for the module while loaded and available */
 	static FHttpModule* Singleton;
+	/** true once StartupModule has been called at least once; prevents Get() from auto-loading after shutdown */
+	static bool bWasEverLoaded;
 	/** The address to use for proxy, in format IPADDRESS:PORT */
 	FString ProxyAddress;
 	/** The domains which won't use proxy even if the ProxyAddress is set, in format "127.0.0.1,localhost,example.com" */
