ios virtual bundle crash fix
- always normalize child bundle relative paths when calling into file system

#Synced-CL 52250858
#Synced-CL 52250858

diff --git a/Engine/Source/Runtime/Online/BuildPatchServices/Private/BuildPatchUtil.cpp b/Engine/Source/Runtime/Online/BuildPatchServices/Private/BuildPatchUtil.cpp
--- a/Engine/Source/Runtime/Online/BuildPatchServices/Private/BuildPatchUtil.cpp
+++ b/Engine/Source/Runtime/Online/BuildPatchServices/Private/BuildPatchUtil.cpp
@@ -318,8 +318,12 @@ FString FBuildPatchUtils::ResolveInstallationFileName(const FString& BaseInstall
 {
 	const FString* PerFileSubdir = PerFileSubdirectories.Find(BuildFilename);
 	if (PerFileSubdir)
-	{
-		return BaseInstallationDirectory / *PerFileSubdir / BuildFilename;
+	{		
+		// Build a relative path from install directory to subdirectory, then resolve it back to absolute path
+		FString CombinedPath = FPaths::Combine(BaseInstallationDirectory, *PerFileSubdir, BuildFilename);
+		FPaths::NormalizeFilename(CombinedPath);
+		FPaths::CollapseRelativeDirectories(CombinedPath, true);
+		return CombinedPath;
 	}
 	return BaseInstallationDirectory / BuildFilename;
 }
