Add HTTPManager skip flush for Android ANR workaround

#android

diff --git a/Engine/Source/Runtime/Online/HTTP/Private/HttpManager.cpp b/Engine/Source/Runtime/Online/HTTP/Private/HttpManager.cpp
--- a/Engine/Source/Runtime/Online/HTTP/Private/HttpManager.cpp
+++ b/Engine/Source/Runtime/Online/HTTP/Private/HttpManager.cpp
@@ -54,6 +54,12 @@ TAutoConsoleVariable<bool> CVarHttpRemoveRequestUsingHttpThreadPolicyOnHttpThrea
 	ECVF_SaveForNextBoot
 );
 
+TAutoConsoleVariable<bool> CVarHttpSkipShutdownFlush(
+	TEXT("http.SkipShutdownFlush"),
+	false,
+	TEXT("When enabled, shutdown does not call flush to avoid waiting past allowed shutdown timing")
+);
+
 TAutoConsoleVariable<bool> CVarHttpFlushTriggerFatalLogDuringCompleteDelegateEnabled(
 	TEXT("http.FlushTriggerFatalLogDuringCompleteDelegateEnabled"),
 	true,
@@ -171,7 +177,10 @@ void FHttpManager::Shutdown()
 	RequestCompletedDelegate.Unbind();
 
 	// Flush all requests
-	Flush(EHttpFlushReason::Shutdown);
+	if (!CVarHttpSkipShutdownFlush.GetValueOnAnyThread())
+	{
+		Flush(EHttpFlushReason::Shutdown);
+	}
 }
 
 bool FHttpManager::HasAnyBoundDelegate() const
