[HttpServer] Fix the conflict merge mistake: Make sure to not set wrong content-length again when use stream response before writing response.

diff --git a/Engine/Source/Runtime/Online/HTTPServer/Private/HttpConnectionResponseWriteContext.cpp b/Engine/Source/Runtime/Online/HTTPServer/Private/HttpConnectionResponseWriteContext.cpp
--- a/Engine/Source/Runtime/Online/HTTPServer/Private/HttpConnectionResponseWriteContext.cpp
+++ b/Engine/Source/Runtime/Online/HTTPServer/Private/HttpConnectionResponseWriteContext.cpp
@@ -28,7 +28,7 @@ void FHttpConnectionResponseWriteContext::ResetContext(TUniquePtr<FHttpServerRes
 	{
 		const bool bResponseIsStream = EnumHasAnyFlags(Response->Flags, EHttpServerResponseFlags::MultipleWriteStream);
 
-		if (!bResponseIsStream || !Response->StreamingBodyQueue.IsValid())
+		if (!bResponseIsStream && !Response->StreamingBodyQueue.IsValid())
 		{
 			// Affix Content-Length Header
 			TArray<FString> ContentLengthValue = { FString::FromInt(Response->Body.Num()) };
