Fixing DownlinkBandwidthManger allocations and Engine testing

#tests Ran Automated testing through the editor and compared mathmatical reporting with external known calculation
#rnx

#Synced-CL 52374847

diff --git a/Engine/Source/Runtime/Online/Experimental/DownlinkBandwidthManager/Private/DownlinkBandwidthManager.cpp b/Engine/Source/Runtime/Online/Experimental/DownlinkBandwidthManager/Private/DownlinkBandwidthManager.cpp
--- a/Engine/Source/Runtime/Online/Experimental/DownlinkBandwidthManager/Private/DownlinkBandwidthManager.cpp
+++ b/Engine/Source/Runtime/Online/Experimental/DownlinkBandwidthManager/Private/DownlinkBandwidthManager.cpp
@@ -170,7 +170,7 @@ UE::BandwidthManager::FServiceHandleID FDownlinkBandwidthManager::RegisterMonito
 			uint8 ValueA = (CurrentGameplayPriorityState == UE::BandwidthManager::EBandwidthDistributionPriorityState::IdleGameplay) ? ServiceA.ConfigData.IdleGameplayPriority : ServiceA.ConfigData.InGameplayPriority;
 			uint8 ValueB = (CurrentGameplayPriorityState == UE::BandwidthManager::EBandwidthDistributionPriorityState::IdleGameplay) ? ServiceB.ConfigData.IdleGameplayPriority : ServiceB.ConfigData.InGameplayPriority;
 			// Sorts in ascending priority, so lower numerical (higher listed) services can be allotted to first
-			return ValueA > ValueB;
+			return ValueA < ValueB;
 		});
 
 	if (StartActive)
@@ -320,6 +320,7 @@ void FDownlinkBandwidthManager::GeneratePriorityBucketMap()
 	}
 	
 	// we do not need an exponent larger than 255, so this lets us make sure that value is not too large and within the limits of the formula
+	const int32 InitialBandwidthToBudget = BandwidthBudgetBytes;
 	uint8 WeightBiasExponent = FMath::Max<uint8>(1, static_cast<uint8>(DLBW::CVarDLBWPriorityWeightExponent.GetValueOnAnyThread()));
 	const int32 BucketCount = PriorityBucketMapping.Num();
 	int32 BucketNumber = 0;
@@ -356,7 +357,7 @@ void FDownlinkBandwidthManager::GeneratePriorityBucketMap()
 			OverallReduction += WeightedPercentage;
 
 			// This is why we are tracking the remainder by bandwidth and not used percentage
-			int32 WeightedBandwidth = static_cast<int32>(static_cast<double>(BandwidthBudgetBytes) * WeightedPercentage);
+			int32 WeightedBandwidth = static_cast<int32>(static_cast<double>(InitialBandwidthToBudget) * WeightedPercentage);
 			BandwidthBudgetBytes -= FMath::Min<int32>(WeightedBandwidth, BandwidthBudgetBytes);
 			// ^^^ Prevents any potential negative roll over
 			BucketRef.BandwidthAllocationForPriority = WeightedBandwidth;
@@ -393,9 +394,9 @@ void FDownlinkBandwidthManager::DistributeBandwidth()
 		if (!ServiceRef.bActiveConnection)
 		{
 			//if not active, then don't distribute bandwidth for it
-			ServiceRef.AllotedBandwidthKbps = 0;
 			if (ServiceRef.OnBandwidthUpdate.IsSet() && ServiceRef.AllotedBandwidthKbps > 0)
 			{
+				ServiceRef.AllotedBandwidthKbps = 0;
 				ServiceRef.OnBandwidthUpdate(MoveTemp(ServiceRef.AllotedBandwidthKbps));
 			}
 			continue;
@@ -478,4 +479,4 @@ void FDownlinkBandwidthManager::Log()
 	}
 
 	UE_LOGF(LogDownlinkBandwidthManager, Log, "FDownlinkBandwidthManager Logging Complete.");
-}
\ No newline at end of file
+}
