Fix missing errno include in BSDSocketTypesPrivate.h for non-Windows platforms

BSDSocketTypesPrivate.h uses errno in GetLastErrorString() but never
explicitly includes <cerrno>, so Unity builds which don't happen to also include
it fail.  Added #include <cerrno> in the non-Windows #else path alongside the
other POSIX headers.

#domain sdk-build
#rnx
#Synced-CL 0

diff --git a/Engine/Source/Runtime/Online/EventLoop/Private/EventLoop/BSDSocket/BSDSocketTypesPrivate.h b/Engine/Source/Runtime/Online/EventLoop/Private/EventLoop/BSDSocket/BSDSocketTypesPrivate.h
--- a/Engine/Source/Runtime/Online/EventLoop/Private/EventLoop/BSDSocket/BSDSocketTypesPrivate.h
+++ b/Engine/Source/Runtime/Online/EventLoop/Private/EventLoop/BSDSocket/BSDSocketTypesPrivate.h
@@ -23,6 +23,7 @@ typedef int32 SOCKLEN;
 #if HAS_EVENTLOOP_PLATFORM_BSD_SOCKET_HEADER
 #include COMPILED_PLATFORM_HEADER(BSDSocketTypesPrivate.h)
 #else // HAS_EVENTLOOP_PLATFORM_BSD_SOCKET_HEADER
+#include <cerrno>
 #include <unistd.h>
 #include <sys/socket.h>
 #if PLATFORM_HAS_BSD_SOCKET_FEATURE_IOCTL
