Soft Launch Enforcement for Bandwidth Distribution Manager

#rnx

#Synced-CL 52374847

diff --git a/Engine/Config/BaseEngine.ini b/Engine/Config/BaseEngine.ini
--- a/Engine/Config/BaseEngine.ini
+++ b/Engine/Config/BaseEngine.ini
@@ -3985,3 +3985,7 @@ bUseSimplifiedUserModel=True
 +HotfixableAssetClasses=/Script/Engine.CurveFloat
 +HotfixableAssetClasses=/Script/Engine.CurveVector
 +HotfixableAssetClasses=/Script/Engine.CurveLinearColor
+
+[/Script/DownlinkBandwidthManager.DownlinkBandwidthManagerConfig]
+RolloutPercentage=0
+EnforceDistributionAllocation=False
\ No newline at end of file
diff --git a/Engine/Source/Runtime/Experimental/IoStore/OnDemand/Private/OnDemandHttpIoDispacher.cpp b/Engine/Source/Runtime/Experimental/IoStore/OnDemand/Private/OnDemandHttpIoDispacher.cpp
--- a/Engine/Source/Runtime/Experimental/IoStore/OnDemand/Private/OnDemandHttpIoDispacher.cpp
+++ b/Engine/Source/Runtime/Experimental/IoStore/OnDemand/Private/OnDemandHttpIoDispacher.cpp
@@ -384,7 +384,7 @@ int32 GetHttpRateLimitKiBPerSecond()
 	}
 #endif // UE_ALLOW_DISABLE_HTTP_LIMITS
 
-	return (GUseDownlinkBandwidthManager) ? DLBWHttpRateLimitKiBPerSecond.load() : GIasHttpRateLimitKiBPerSecond;
+	return (GUseDownlinkBandwidthManager && FDownlinkBandwidthManager::HasValueEnforcementPassed()) ? DLBWHttpRateLimitKiBPerSecond.load() : GIasHttpRateLimitKiBPerSecond;
 }
 
 ////////////////////////////////////////////////////////////////////////////////
@@ -1525,7 +1525,7 @@ uint32 FHttpDispatcher::Run()
 
 	for (;;)
 	{
-		if (GUseDownlinkBandwidthManager && BandwidthManagerServiceHandle == UE_INVALID_BANDWIDTH_SERVICE_HANDLE)
+		if (FDownlinkBandwidthManager::HasLocalUserPassedRolloutCheck() && BandwidthManagerServiceHandle == UE_INVALID_BANDWIDTH_SERVICE_HANDLE)
 		{
 			FDownlinkBandwidthManager& DBManager = FDownlinkBandwidthManager::Get();
 			UE::BandwidthManager::FManagedServiceConfigurableData BandwidthConfig;
@@ -1541,7 +1541,7 @@ uint32 FHttpDispatcher::Run()
 					DLBWHttpRateLimitKiBPerSecond = static_cast<int32>((float)AllottedBandwidthKbps * KilobyteToKibibyte);
 				});
 		}
-		else if (GUseDownlinkBandwidthManager)
+		else if (FDownlinkBandwidthManager::HasLocalUserPassedRolloutCheck())
 		{
 			FDownlinkBandwidthManager& DBManager = FDownlinkBandwidthManager::Get();
 			DBManager.SetActivationForService(BandwidthManagerServiceHandle, true);
@@ -1593,7 +1593,7 @@ void FHttpDispatcher::Stop()
 		WakeUp->Trigger();
 	}
 
-	if (GUseDownlinkBandwidthManager)
+	if (FDownlinkBandwidthManager::HasLocalUserPassedRolloutCheck())
 	{
 		FDownlinkBandwidthManager& DBManager = FDownlinkBandwidthManager::Get();
 		DBManager.SetActivationForService(BandwidthManagerServiceHandle, false);
diff --git a/Engine/Source/Runtime/Online/Experimental/DownlinkBandwidthManager/Private/DownlinkBandwidthManager.cpp b/Engine/Source/Runtime/Online/Experimental/DownlinkBandwidthManager/Private/DownlinkBandwidthManager.cpp
--- a/Engine/Source/Runtime/Online/Experimental/DownlinkBandwidthManager/Private/DownlinkBandwidthManager.cpp
+++ b/Engine/Source/Runtime/Online/Experimental/DownlinkBandwidthManager/Private/DownlinkBandwidthManager.cpp
@@ -1,9 +1,14 @@
 // Copyright Epic Games, Inc. All Rights Reserved.
 
 #include "DownlinkBandwidthManager.h"
+#include "DownlinkBandwidthManagerConfig.h"
+#include "BandwidthMeasurementTool.h"
 #include "HAL/IConsoleManager.h"
 #include "Math/NumericLimits.h"
-#include "BandwidthMeasurementTool.h"
+#include "Misc/ConfigCacheIni.h"
+#include "UObject/UObjectGlobals.h"
+
+#include <random>
 
 DEFINE_LOG_CATEGORY(LogDownlinkBandwidthManager);
 
@@ -21,7 +26,7 @@ namespace DLBW
 	static FAutoConsoleVariableRef CVarDLBWMinRequiredBandwidthDuringGameplay(
 		TEXT("net.DLBW.MinGameplayBandwidth"),
 		GMinRequiredBandwidthDuringGameplay,
-		TEXT("Default value to use for bandwidth cap distribution in Kbps. Example: Average American home internet is 200-300Mbps, so we would use 200,000 - 300,000 for this value."),
+		TEXT("This is a failsafe value, meant to act as the reserved size minimum we want to set aside should the value collection from the UNetDriver be too small."),
 		ECVF_Default);
 
 	static TAutoConsoleVariable<int32> CVarDLBWPriorityWeightExponent(
@@ -36,13 +41,27 @@ namespace DLBW
 		TEXT("Delay Time, per frame cycle to poll for changes needed in distribution."),
 		ECVF_Default);
 
-	float GBandwidthPercentageToUse = 0.8f; 
+	static float GBandwidthPercentageToUse = 0.8f; 
 	static FAutoConsoleVariableRef CVarDLBWBandwidthPercentToUse(
 		TEXT("net.DLBW.BandwidthPercentage"),
 		GBandwidthPercentageToUse,
 		TEXT("0 - 1 value used to control the amount of measured bandwidth the manager will utilize."),
 		ECVF_Default);
 
+	static int32 GRolloutPercentage = -1;
+	static FAutoConsoleVariableRef CVarDLBWRolloutPercentage(
+		TEXT("net.DLBW.RolloutChance"),
+		GRolloutPercentage,
+		TEXT("A value from 0-100 that will be a percentage chance for the user to use values from the Distribution Manager. This is rolled per launch! For more guaranteed testing, enable the individual services for DLBW directly."),
+		ECVF_Default);
+
+	static bool GEnforceDistributionAllocation = false;
+	static FAutoConsoleVariableRef CVarDLBWEnforcement(
+		TEXT("net.DLBW.Enforcement"),
+		GEnforceDistributionAllocation,
+		TEXT("A true/false value used to denote whether we are in the rollout phase for applying the distributor's calculated allocations to subscribed systems."),
+		ECVF_Default);
+
 	// Command to print logging information
 	FAutoConsoleCommand FDLBWShowBandwidthTelemetry(
 		TEXT("net.DLBW.PrintTelemetry"),
@@ -87,6 +106,18 @@ FDownlinkBandwidthManager& FDownlinkBandwidthManager::Get()
 	return *Instance.Get();
 }
 
+bool FDownlinkBandwidthManager::HasLocalUserPassedRolloutCheck()
+{
+	FDownlinkBandwidthManager& DLBWManager = FDownlinkBandwidthManager::Get();
+	return static_cast<float>(DLBW::GRolloutPercentage) >= DLBWManager.RolloutPercentage;
+}
+
+bool FDownlinkBandwidthManager::HasValueEnforcementPassed()
+{
+	FDownlinkBandwidthManager& DLBWManager = FDownlinkBandwidthManager::Get();
+	return DLBW::GEnforceDistributionAllocation && HasLocalUserPassedRolloutCheck();
+}
+
 /// Member Functions
 FDownlinkBandwidthManager::FDownlinkBandwidthManager()
 {
@@ -100,6 +131,22 @@ FDownlinkBandwidthManager::~FDownlinkBandwidthManager()
 
 void FDownlinkBandwidthManager::Initialize()
 {
+	// Get Rollout Data
+	if (UClass* ConfigClass = UDownlinkBandwidthManagerConfig::StaticClass())
+	{
+		if (UDownlinkBandwidthManagerConfig* Config = ConfigClass->GetDefaultObject<UDownlinkBandwidthManagerConfig>())
+		{
+			Config->SetConsoleVariablesFromConfigurables();
+		}
+	}
+
+	std::random_device Device;
+	std::mt19937 Generator(Device());
+	std::uniform_real_distribution<float> Distribution(0.0f, 100.0f);
+
+	// A value in the range [0..100), EG. inclusive of 0, exclusive of 100.
+	RolloutPercentage = Distribution(Generator);
+
 	// Set an initial value
 	SetIncomingBandwidthCap(DLBW::GDefaultBandwidthToAllocate);
 	
@@ -423,7 +470,7 @@ void FDownlinkBandwidthManager::DistributeBandwidth()
 			// reduce total amount by amount distributed
 			BandwidthBudgetBytes -= ServiceRef.AllotedBandwidthKbps * UE::BandwidthManager::KilobitsToBytesConversionValue;
 
-			if (ServiceRef.OnBandwidthUpdate.IsSet() && ServiceRef.AllotedBandwidthKbps > 0)
+			if (DLBW::GEnforceDistributionAllocation && ServiceRef.OnBandwidthUpdate.IsSet() && ServiceRef.AllotedBandwidthKbps > 0)
 			{
 				ServiceRef.OnBandwidthUpdate(MoveTemp(ServiceRef.AllotedBandwidthKbps));
 			}
diff --git a/Engine/Source/Runtime/Online/Experimental/DownlinkBandwidthManager/Private/DownlinkBandwidthManagerConfig.cpp b/Engine/Source/Runtime/Online/Experimental/DownlinkBandwidthManager/Private/DownlinkBandwidthManagerConfig.cpp
new file mode 100644
--- /dev/null
+++ b/Engine/Source/Runtime/Online/Experimental/DownlinkBandwidthManager/Private/DownlinkBandwidthManagerConfig.cpp
@@ -0,0 +1,25 @@
+// Copyright Epic Games, Inc. All Rights Reserved.
+
+#include "DownlinkBandwidthManagerConfig.h"
+#include "HAL/IConsoleManager.h"
+
+
+void UDownlinkBandwidthManagerConfig::PostReloadConfig(class FProperty* PropertyThatWasLoaded)
+{
+	SetConsoleVariablesFromConfigurables();
+}
+
+void UDownlinkBandwidthManagerConfig::SetConsoleVariablesFromConfigurables()
+{
+	RolloutPercentage = FMath::Clamp(RolloutPercentage, 0, 100);
+
+	if (IConsoleVariable* ConsoleVariable = IConsoleManager::Get().FindConsoleVariable(TEXT("net.DLBW.RolloutChance")))
+	{
+		ConsoleVariable->Set(RolloutPercentage, ECVF_SetByCode);
+	}
+
+	if (IConsoleVariable* ConsoleVariable = IConsoleManager::Get().FindConsoleVariable(TEXT("net.DLBW.Enforcement")))
+	{
+		ConsoleVariable->Set(EnforceDistributionAllocation, ECVF_SetByCode);
+	}
+}
diff --git a/Engine/Source/Runtime/Online/Experimental/DownlinkBandwidthManager/Public/DownlinkBandwidthManager.h b/Engine/Source/Runtime/Online/Experimental/DownlinkBandwidthManager/Public/DownlinkBandwidthManager.h
--- a/Engine/Source/Runtime/Online/Experimental/DownlinkBandwidthManager/Public/DownlinkBandwidthManager.h
+++ b/Engine/Source/Runtime/Online/Experimental/DownlinkBandwidthManager/Public/DownlinkBandwidthManager.h
@@ -1,13 +1,14 @@
 // Copyright Epic Games, Inc. All Rights Reserved.
 #pragma once
 #include "Containers/Ticker.h"
+#include "CoreMinimal.h"
 
 #define UE_API DOWNLINKBANDWIDTHMANAGER_API
 #define UE_INVALID_BANDWIDTH_SERVICE_HANDLE -1
 
 UE_API DECLARE_LOG_CATEGORY_EXTERN(LogDownlinkBandwidthManager, Log, All);
 
-// Useage for this callback will receive a parameter in Kbps
+// Usage for this callback will receive a parameter in Kbps
 using FOnBandwidthUpdate = TUniqueFunction<void(int32&&)>;
 
 namespace UE::BandwidthManager
@@ -97,6 +98,12 @@ public:
 	// --- Singleton access ---
 	UE_API static FDownlinkBandwidthManager& Get();
 
+	// Check to see if the probability check for roll out has passed/failed
+	UE_API static bool HasLocalUserPassedRolloutCheck();
+
+	// Check to see if the probability check for roll out has passed/failed and Enforcement is enabled
+	UE_API static bool HasValueEnforcementPassed();
+
 	// --- Lifecycle (called by module) ---
 	void Initialize();
 	void Shutdown();
@@ -170,6 +177,9 @@ private:
 
 	// The Managers perceived state of gameplay to consider for distribution
 	UE::BandwidthManager::EBandwidthDistributionPriorityState CurrentGameplayPriorityState = UE::BandwidthManager::EBandwidthDistributionPriorityState::IdleGameplay;
+
+	// Randomized chance roll at initialization to determine if the manager should rollout for this client's launch
+	float RolloutPercentage = 0.0f;
 };
 
 #undef UE_API
diff --git a/Engine/Source/Runtime/Online/Experimental/DownlinkBandwidthManager/Public/DownlinkBandwidthManagerConfig.h b/Engine/Source/Runtime/Online/Experimental/DownlinkBandwidthManager/Public/DownlinkBandwidthManagerConfig.h
new file mode 100644
--- /dev/null
+++ b/Engine/Source/Runtime/Online/Experimental/DownlinkBandwidthManager/Public/DownlinkBandwidthManagerConfig.h
@@ -0,0 +1,29 @@
+// Copyright Epic Games, Inc. All Rights Reserved.
+#pragma once
+#include "CoreMinimal.h"
+
+#include "DownlinkBandwidthManagerConfig.generated.h"
+
+#define UE_API DOWNLINKBANDWIDTHMANAGER_API
+
+
+UCLASS(config = Engine, MinimalAPI)
+class UDownlinkBandwidthManagerConfig : public UObject
+{
+	GENERATED_BODY()
+public:
+	UDownlinkBandwidthManagerConfig() {};
+
+	UE_API virtual void PostReloadConfig(class FProperty* PropertyThatWasLoaded) override;
+
+	void SetConsoleVariablesFromConfigurables();
+
+	// The pass percentage for the user to run the Distribution Manager.
+	UPROPERTY(Config)
+	int32 RolloutPercentage = 0;
+
+	// Enables the application of the distributor's calculated allocations to subscribed systems if percentage check has passed
+	UPROPERTY(Config)
+	bool EnforceDistributionAllocation = false;
+};
+#undef UE_API
\ No newline at end of file
